Move SStatusBar to new StatusBarWidget developer module so it can be used in programs

diff --git a/Engine/Source/Developer/ToolWidgets/Private/SWidgetDrawer.cpp b/Engine/Source/Developer/ToolWidgets/Private/SWidgetDrawer.cpp
new file mode 100644
--- /dev/null
+++ b/Engine/Source/Developer/ToolWidgets/Private/SWidgetDrawer.cpp
@@ -0,0 +1,675 @@
+// Copyright Epic Games, Inc. All Rights Reserved.
+
+#include "SWidgetDrawer.h"
+#include "Widgets/Images/SImage.h"
+#include "Widgets/Input/SButton.h"
+#include "Widgets/Layout/SBorder.h"
+#include "Widgets/Text/SRichTextBlock.h"
+#include "Widgets/Docking/SDockTab.h"
+#include "Framework/Application/SlateApplication.h"
+#include "InputCoreTypes.h"
+#include "Misc/ConfigCacheIni.h"
+#include "Widgets/Notifications/SProgressBar.h"
+#include "Widgets/Notifications/SNotificationList.h"
+#include "Framework/Notifications/NotificationManager.h"
+
+#define LOCTEXT_NAMESPACE "StatusBar"
+
+namespace StatusBarNotificationConstants
+{
+	// How long progress notification toasts should appear for
+	const float NotificationExpireTime = 5.0f;
+
+	const float NotificationFadeDuration = .15f;
+
+	// Delay before a progress notification becomes visible. This is to avoid the status bar to animate and flicker from short lived notifications. 
+	const double NotificationDelay = 3.0;
+}
+
+class SDrawerOverlay : public SCompoundWidget
+{
+	SLATE_BEGIN_ARGS(SDrawerOverlay)
+	{
+		_Clipping = EWidgetClipping::ClipToBounds;
+		_ShadowOffset = FVector2D(10.0f, 20.0f);
+	}
+		SLATE_DEFAULT_SLOT(FArguments, Content)
+		SLATE_ARGUMENT(float, MinDrawerHeight)
+		SLATE_ARGUMENT(float, MaxDrawerHeight)
+		SLATE_ARGUMENT(float, TargetDrawerHeight)
+		SLATE_EVENT(FOnStatusBarDrawerTargetHeightChanged, OnTargetHeightChanged)
+		SLATE_EVENT(FSimpleDelegate, OnDismissComplete)
+		SLATE_ARGUMENT(FVector2D, ShadowOffset)
+	SLATE_END_ARGS()
+
+	~SDrawerOverlay()
+	{
+		FSlateThrottleManager::Get().LeaveResponsiveMode(AnimationThrottle);
+	}
+
+	void Construct(const FArguments& InArgs)
+	{
+		CurrentHeight = 0;
+
+		ShadowOffset = InArgs._ShadowOffset;
+		ExpanderSize = 5.0f;
+
+		SplitterStyle = &FAppStyle::Get().GetWidgetStyle<FSplitterStyle>("Splitter");
+
+		MinHeight = InArgs._MinDrawerHeight;
+
+		MaxHeight = InArgs._MaxDrawerHeight;
+
+		TargetHeight = FMath::Clamp(InArgs._TargetDrawerHeight, MinHeight, MaxHeight);
+
+		OnTargetHeightChanged = InArgs._OnTargetHeightChanged;
+
+		BackgroundBrush = FAppStyle::Get().GetBrush("StatusBar.DrawerBackground");
+		ShadowBrush = FAppStyle::Get().GetBrush("StatusBar.DrawerShadow");
+		BorderBrush = FAppStyle::Get().GetBrush("Docking.Sidebar.Border");
+
+		bIsResizeHandleHovered = false;
+		bIsResizing = false;
+
+		OnDismissComplete = InArgs._OnDismissComplete;
+
+		DrawerEasingCurve = FCurveSequence(0.0f, 0.15f, ECurveEaseFunction::QuadOut);
+
+		ChildSlot
+		[
+			InArgs._Content.Widget
+		];
+	}
+
+	void UpdateHeightInterp(float InAlpha)
+	{
+		float NewHeight = FMath::Lerp(0.0f, TargetHeight, InAlpha);
+
+		SetHeight(NewHeight);
+	}
+
+	virtual bool SupportsKeyboardFocus() const override
+	{
+		return true;
+	}
+
+	virtual FVector2D ComputeDesiredSize(float) const
+	{
+		return FVector2D(1.0f, TargetHeight + ShadowOffset.Y);
+	}
+
+	virtual void OnArrangeChildren(const FGeometry& AllottedGeometry, FArrangedChildren& ArrangedChildren) const override
+	{
+		const EVisibility ChildVisibility = ChildSlot.GetWidget()->GetVisibility();
+		if (ArrangedChildren.Accepts(ChildVisibility))
+		{
+			ArrangedChildren.AddWidget(
+				AllottedGeometry.MakeChild(
+					ChildSlot.GetWidget(),
+					ShadowOffset,
+					FVector2D(AllottedGeometry.GetLocalSize().X - (ShadowOffset.X * 2), TargetHeight)
+				)
+			);
+		}
+	}
+
+	virtual FReply OnMouseButtonDown(const FGeometry& AllottedGeometry, const FPointerEvent& MouseEvent) override
+	{
+		FReply Reply = FReply::Unhandled();
+		if (MouseEvent.GetEffectingButton() == EKeys::LeftMouseButton)
+		{
+			const FGeometry RenderTransformedChildGeometry = GetRenderTransformedGeometry(AllottedGeometry);
+			const FGeometry ResizeHandleGeometry = GetResizeHandleGeometry(AllottedGeometry);
+
+			if (ResizeHandleGeometry.IsUnderLocation(MouseEvent.GetScreenSpacePosition()))
+			{
+				bIsResizing = true;
+				InitialResizeGeometry = ResizeHandleGeometry;
+				InitialHeightAtResize = CurrentHeight;
+				ResizeThrottleHandle = FSlateThrottleManager::Get().EnterResponsiveMode();
+
+				Reply = FReply::Handled().CaptureMouse(SharedThis(this));
+			}
+		}
+
+		return Reply;
+
+	}
+
+	FReply OnMouseButtonUp(const FGeometry& AllottedGeometry, const FPointerEvent& MouseEvent) override
+	{
+		if (MouseEvent.GetEffectingButton() == EKeys::LeftMouseButton && bIsResizing == true)
+		{
+			bIsResizing = false;
+			FSlateThrottleManager::Get().LeaveResponsiveMode(ResizeThrottleHandle);
+
+			OnTargetHeightChanged.ExecuteIfBound(TargetHeight);
+			return FReply::Handled().ReleaseMouseCapture();
+		}
+		return FReply::Unhandled();
+	}
+
+	FReply OnMouseMove(const FGeometry& AllottedGeometry, const FPointerEvent& MouseEvent) override
+	{
+		const FGeometry ResizeHandleGeometry = GetResizeHandleGeometry(AllottedGeometry);
+
+		bIsResizeHandleHovered = ResizeHandleGeometry.IsUnderLocation(MouseEvent.GetScreenSpacePosition());
+
+		if (bIsResizing && this->HasMouseCapture() && !MouseEvent.GetCursorDelta().IsZero())
+		{
+			const FVector2D LocalMousePos = InitialResizeGeometry.AbsoluteToLocal(MouseEvent.GetScreenSpacePosition());
+			const float DeltaHeight = (InitialResizeGeometry.GetLocalPositionAtCoordinates(FVector2D::ZeroVector) - LocalMousePos).Y;
+
+			TargetHeight = FMath::Clamp(InitialHeightAtResize + DeltaHeight, MinHeight, MaxHeight);
+			SetHeight(InitialHeightAtResize + DeltaHeight);
+
+
+			return FReply::Handled();
+		}
+		else
+		{
+			return FReply::Unhandled();
+		}
+	}
+
+	virtual void OnMouseLeave(const FPointerEvent& MouseEvent) override
+	{
+		SCompoundWidget::OnMouseLeave(MouseEvent);
+
+		bIsResizeHandleHovered = false;
+	}
+
+	FCursorReply OnCursorQuery(const FGeometry& MyGeometry, const FPointerEvent& CursorEvent) const override
+	{
+		return bIsResizing || bIsResizeHandleHovered ? FCursorReply::Cursor(EMouseCursor::ResizeUpDown) : FCursorReply::Unhandled();
+	}
+
+	virtual int32 OnPaint(const FPaintArgs& Args, const FGeometry& AllottedGeometry, const FSlateRect& MyCullingRect, FSlateWindowElementList& OutDrawElements, int32 LayerId, const FWidgetStyle& InWidgetStyle, bool bParentEnabled) const override
+	{
+		static FSlateColor ShadowColor = FAppStyle::Get().GetSlateColor("Colors.Foldout");
+
+		const FGeometry RenderTransformedChildGeometry = GetRenderTransformedGeometry(AllottedGeometry);
+		const FGeometry ResizeHandleGeometry = GetResizeHandleGeometry(AllottedGeometry);
+
+		// Draw the resize handle
+		if (bIsResizing || bIsResizeHandleHovered)
+		{
+			const FSlateBrush* SplitterBrush = &SplitterStyle->HandleHighlightBrush;
+			FSlateDrawElement::MakeBox(
+				OutDrawElements,
+				LayerId,
+				ResizeHandleGeometry.ToPaintGeometry(),
+				SplitterBrush,
+				ESlateDrawEffect::None,
+				SplitterBrush->GetTint(InWidgetStyle));
+		}
+
+		// Top Shadow
+		FSlateDrawElement::MakeBox(
+			OutDrawElements,
+			LayerId,
+			RenderTransformedChildGeometry.ToPaintGeometry(),
+			ShadowBrush,
+			ESlateDrawEffect::None,
+			ShadowBrush->GetTint(InWidgetStyle));
+
+		// Background
+		FSlateDrawElement::MakeBox(
+			OutDrawElements,
+			LayerId,
+			RenderTransformedChildGeometry.ToPaintGeometry(ShadowOffset, FVector2D(AllottedGeometry.GetLocalSize().X - (ShadowOffset.X * 2), TargetHeight)),
+			BackgroundBrush,
+			ESlateDrawEffect::None,
+			BackgroundBrush->GetTint(InWidgetStyle));
+
+		int32 OutLayerId = SCompoundWidget::OnPaint(Args, RenderTransformedChildGeometry, MyCullingRect, OutDrawElements, LayerId, InWidgetStyle, bParentEnabled);
+
+		// Bottom shadow
+		FSlateDrawElement::MakeBox(
+			OutDrawElements,
+			OutLayerId,
+			AllottedGeometry.ToPaintGeometry(FVector2D(0.0f, AllottedGeometry.GetLocalSize().Y - ShadowOffset.Y), FVector2D(AllottedGeometry.GetLocalSize().X, ShadowOffset.Y)),
+			ShadowBrush,
+			ESlateDrawEffect::None,
+			ShadowBrush->GetTint(InWidgetStyle));
+
+
+		FSlateDrawElement::MakeBox(
+			OutDrawElements,
+			OutLayerId+1,
+			RenderTransformedChildGeometry.ToPaintGeometry(ShadowOffset, FVector2D(AllottedGeometry.GetLocalSize().X - (ShadowOffset.X * 2), TargetHeight)),
+			BorderBrush,
+			ESlateDrawEffect::None,
+			BorderBrush->GetTint(InWidgetStyle));
+
+		return OutLayerId+1;
+
+	}
+
+	void Open()
+	{
+		DrawerEasingCurve.Play(AsShared(), false, DrawerEasingCurve.IsPlaying() ? DrawerEasingCurve.GetSequenceTime() : 0.0f, false);
+
+		if (!DrawerOpenCloseTimer.IsValid())
+		{
+			AnimationThrottle = FSlateThrottleManager::Get().EnterResponsiveMode();
+			DrawerOpenCloseTimer = RegisterActiveTimer(0.0f, FWidgetActiveTimerDelegate::CreateSP(this, &SDrawerOverlay::UpdateDrawerAnimation));
+		}
+	}
+
+	void Dismiss()
+	{
+		if (DrawerEasingCurve.IsForward())
+		{
+			DrawerEasingCurve.Reverse();
+		}
+
+		if (!DrawerOpenCloseTimer.IsValid())
+		{
+			AnimationThrottle = FSlateThrottleManager::Get().EnterResponsiveMode();
+			DrawerOpenCloseTimer = RegisterActiveTimer(0.0f, FWidgetActiveTimerDelegate::CreateSP(this, &SDrawerOverlay::UpdateDrawerAnimation));
+		}
+	}
+private:
+	FGeometry GetRenderTransformedGeometry(const FGeometry& AllottedGeometry) const
+	{
+		return AllottedGeometry.MakeChild(FSlateRenderTransform(FVector2D(0.0f, TargetHeight - CurrentHeight)));
+	}
+
+	FGeometry GetResizeHandleGeometry(const FGeometry& AllottedGeometry) const
+	{
+		return GetRenderTransformedGeometry(AllottedGeometry).MakeChild(ShadowOffset - FVector2D(0.0f, ExpanderSize), FVector2D(AllottedGeometry.GetLocalSize().X-ShadowOffset.X*2, ExpanderSize));
+	}
+
+	void SetHeight(float NewHeight)
+	{
+		CurrentHeight = FMath::Clamp(NewHeight, MinHeight, TargetHeight);
+	}
+
+	EActiveTimerReturnType UpdateDrawerAnimation(double CurrentTime, float DeltaTime)
+	{
+		UpdateHeightInterp(DrawerEasingCurve.GetLerp());
+
+		if (!DrawerEasingCurve.IsPlaying())
+		{
+			if (DrawerEasingCurve.IsAtStart())
+			{
+				OnDismissComplete.ExecuteIfBound();
+			}
+
+			FSlateThrottleManager::Get().LeaveResponsiveMode(AnimationThrottle);
+			DrawerOpenCloseTimer.Reset();
+			return EActiveTimerReturnType::Stop;
+		}
+
+		return EActiveTimerReturnType::Continue;
+	}
+
+
+private:
+	FGeometry InitialResizeGeometry;
+	TSharedPtr<FActiveTimerHandle> DrawerOpenCloseTimer;
+	FOnStatusBarDrawerTargetHeightChanged OnTargetHeightChanged;
+	FCurveSequence DrawerEasingCurve;
+	FSimpleDelegate OnDismissComplete;
+	const FSlateBrush* BackgroundBrush;
+	const FSlateBrush* ShadowBrush;
+	const FSlateBrush* BorderBrush;
+	const FSplitterStyle* SplitterStyle;
+	FVector2D ShadowOffset;
+	FThrottleRequest AnimationThrottle;
+	FThrottleRequest ResizeThrottleHandle;
+	float ExpanderSize;
+	float CurrentHeight;
+	float MinHeight;
+	float MaxHeight;
+	float TargetHeight; 
+	float InitialHeightAtResize;
+	bool bIsResizing;
+	bool bIsResizeHandleHovered;
+};
+
+SWidgetDrawer::~SWidgetDrawer()
+{
+	// Ensure the content browser is removed if we're being destroyed
+	CloseDrawerImmediately();
+}
+
+void SWidgetDrawer::Construct(const FArguments& InArgs, FName InStatusBarName)
+{
+	DrawerName = InStatusBarName;
+	
+	FSlateApplication::Get().OnFocusChanging().AddSP(this, &SWidgetDrawer::OnGlobalFocusChanging);
+	FGlobalTabmanager::Get()->OnActiveTabChanged_Subscribe(FOnActiveTabChanged::FDelegate::CreateSP(this, &SWidgetDrawer::OnActiveTabChanged));
+	FGlobalTabmanager::Get()->OnTabForegrounded_Subscribe(FOnActiveTabChanged::FDelegate::CreateSP(this, &SWidgetDrawer::OnActiveTabChanged));
+	
+	ChildSlot
+	[
+		SAssignNew(DrawerBox, SHorizontalBox)
+	];
+}
+
+void SWidgetDrawer::OnGlobalFocusChanging(const FFocusEvent& FocusEvent, const FWeakWidgetPath& OldFocusedWidgetPath, const TSharedPtr<SWidget>& OldFocusedWidget, const FWidgetPath& NewFocusedWidgetPath, const TSharedPtr<SWidget>& NewFocusedWidget)
+{
+	// Sometimes when dismissing focus can change which will trigger this again
+	static bool bIsRentrant = false;
+
+	if(!bIsRentrant)
+	{
+		TGuardValue<bool> RentrancyGuard(bIsRentrant, true);
+
+		TSharedRef<SWidget> ThisWidget = AsShared();
+
+		TSharedPtr<SWidget> ActiveDrawerOverlayContent;
+		if (OpenedDrawer.IsValid())
+		{
+			ActiveDrawerOverlayContent = OpenedDrawer.DrawerOverlay;
+		}
+
+		bool bShouldDismiss = false;
+
+		// If we aren't focusing any new widgets, act as if the drawer is in the path 
+		bool bDrawerInPath = NewFocusedWidgetPath.ContainsWidget(ActiveDrawerOverlayContent.Get()) 
+			|| NewFocusedWidgetPath.ContainsWidget(this) 
+			|| NewFocusedWidgetPath.Widgets.Num() == 0;
+
+		// Do not close due to slow tasks as those opening send window activation events
+		if (!GIsSlowTask && !bDrawerInPath && !FSlateApplication::Get().GetActiveModalWindow().IsValid() && ActiveDrawerOverlayContent.IsValid())
+		{
+			if (TSharedPtr<SWidget> MenuHost = FSlateApplication::Get().GetMenuHostWidget())
+			{
+				FWidgetPath MenuHostPath;
+
+				// See if the menu being opened is part of the content browser path and if so the menu should not be dismissed
+				FSlateApplication::Get().GeneratePathToWidgetUnchecked(MenuHost.ToSharedRef(), MenuHostPath, EVisibility::All);
+				if (!MenuHostPath.ContainsWidget(ActiveDrawerOverlayContent.Get()))
+				{
+					bShouldDismiss = true;
+				}
+			}
+			else
+			{
+				bShouldDismiss = true;
+			}
+		}
+
+		if (bShouldDismiss)
+		{
+			DismissDrawer(NewFocusedWidget);
+		}
+	}
+}
+
+void SWidgetDrawer::OnActiveTabChanged(TSharedPtr<SDockTab> PreviouslyActive, TSharedPtr<SDockTab> NewlyActivated)
+{
+	bool bShouldRemoveDrawer = false;
+	if (!PreviouslyActive || !NewlyActivated)
+	{
+		// Remove the content browser if there is some invalid state with the tabs
+		bShouldRemoveDrawer = true;
+	}
+	else if(NewlyActivated->GetTabRole() == ETabRole::MajorTab)
+	{
+		// Remove the content browser if a newly activated tab is a major tab
+		bShouldRemoveDrawer = true;
+	}
+	else if (PreviouslyActive->GetTabManagerPtr() != NewlyActivated->GetTabManagerPtr())
+	{
+		// Remove the content browser if we're switching tab managers (indicates a new status bar is becoming active)
+		bShouldRemoveDrawer = true;
+	}
+
+	if (bShouldRemoveDrawer)
+	{
+		CloseDrawerImmediately();
+	}
+}
+
+TSharedRef<SWidget> SWidgetDrawer::MakeStatusBarDrawerButton(const FWidgetDrawerConfig& Drawer)
+{
+	const FName DrawerId = Drawer.UniqueId;
+
+	const FSlateBrush* StatusBarBackground = FAppStyle::Get().GetBrush("Brushes.Panel");
+
+	TSharedRef<SWidget> DrawerButton = 
+
+		SNew(SBorder)
+		.Padding(FMargin(2.0f, 0.0f))
+		.BorderImage(StatusBarBackground)
+		.Visibility(EVisibility::SelfHitTestInvisible)
+		.VAlign(VAlign_Center)
+		[
+			SNew(SButton)
+			.IsFocusable(false)
+			.ButtonStyle(&FAppStyle::Get().GetWidgetStyle<FButtonStyle>("StatusBar.StatusBarButton"))
+			.OnClicked(this, &SWidgetDrawer::OnDrawerButtonClicked, DrawerId)
+			.ToolTipText(Drawer.ToolTipText)
+			[
+				SNew(SHorizontalBox)
+				+ SHorizontalBox::Slot()
+				.Padding(2.0f)
+				.HAlign(HAlign_Left)
+				.VAlign(VAlign_Center)
+				.AutoWidth()
+				[
+					SNew(SImage)
+					.ColorAndOpacity(FSlateColor::UseForeground())
+					.Image(Drawer.Icon)
+				]
+				+ SHorizontalBox::Slot()
+				.VAlign(VAlign_Center)
+				.Padding(2.0f)
+				[
+					SNew(STextBlock)
+					.TextStyle(&FAppStyle::Get().GetWidgetStyle<FTextBlockStyle>("NormalText"))
+					.Text(Drawer.ButtonText)
+				]
+			]
+		];
+
+
+	if (Drawer.CustomWidget)
+	{
+		return
+			SNew(SHorizontalBox)
+			+ SHorizontalBox::Slot()
+			.Padding(0.0f, 0.0f, 2.0f, 0.0f)
+			.AutoWidth()
+			[
+				DrawerButton
+			]
+			+ SHorizontalBox::Slot()
+			[
+				SNew(SBorder)
+				.Padding(FMargin(2.0f, 0.0f))
+				.BorderImage(StatusBarBackground)
+				.Visibility(EVisibility::SelfHitTestInvisible)
+				.VAlign(VAlign_Center)
+				[
+					Drawer.CustomWidget.ToSharedRef()
+				]
+			];
+	
+	}
+	else
+	{
+		return DrawerButton;
+	}
+}
+
+bool SWidgetDrawer::IsDrawerOpened(const FName DrawerId) const
+{
+	return OpenedDrawer == DrawerId ? true : false;
+}
+
+bool SWidgetDrawer::IsAnyOtherDrawerOpened(const FName DrawerId) const
+{
+	return OpenedDrawer.IsValid() && OpenedDrawer.DrawerId != DrawerId ? true : false;
+}
+
+FString SWidgetDrawer::GetSerializableName() const
+{
+	return DrawerName.GetPlainNameString();
+}
+
+FReply SWidgetDrawer::OnDrawerButtonClicked(const FName DrawerId)
+{
+	if (!IsDrawerOpened(DrawerId))
+	{
+		OpenDrawer(DrawerId);
+	}
+	else
+	{
+		DismissDrawer(nullptr);
+	}
+
+	return FReply::Handled();
+}
+
+void SWidgetDrawer::OnDrawerHeightChanged(float TargetHeight)
+{
+	TSharedPtr<SWindow> MyWindow = OpenedDrawer.WindowWithOverlayContent.Pin();
+
+	// Save the height has a percentage of the screen
+	const float TargetDrawerHeightPct = TargetHeight / (MyWindow->GetSizeInScreen().Y / MyWindow->GetDPIScaleFactor());
+	GConfig->SetFloat(TEXT("DrawerSizes"), *(GetSerializableName() + TEXT(".") + OpenedDrawer.DrawerId.ToString()), TargetDrawerHeightPct, GEditorSettingsIni);
+}
+
+void SWidgetDrawer::CloseDrawerImmediatelyInternal(const FOpenDrawerData& Data)
+{
+	if (Data.IsValid())
+	{
+		TSharedRef<SWidget> DrawerOverlayContent = Data.DrawerOverlay.ToSharedRef();
+
+		// Remove the content browser from the window
+		if (TSharedPtr<SWindow> Window = Data.WindowWithOverlayContent.Pin())
+		{
+			Window->RemoveOverlaySlot(DrawerOverlayContent);
+		}
+	}
+}
+
+void SWidgetDrawer::RegisterDrawer(FWidgetDrawerConfig&& Drawer, int32 SlotIndex)
+{
+	const int32 NumDrawers = RegisteredDrawers.Num();
+	RegisteredDrawers.AddUnique(Drawer);
+
+	if (RegisteredDrawers.Num() > NumDrawers)
+	{
+		DrawerBox->InsertSlot(SlotIndex)
+		.Padding(1.0f, 0.0f)
+		.AutoWidth()
+		[
+			MakeStatusBarDrawerButton(Drawer)
+		];
+	}
+}
+
+void SWidgetDrawer::OpenDrawer(const FName DrawerId)
+{
+	// Close any other open drawer
+	if (OpenedDrawer.DrawerId != DrawerId && DismissingDrawers.IndexOfByKey(DrawerId) == INDEX_NONE)
+	{
+		DismissDrawer(nullptr);
+
+		FWidgetDrawerConfig* DrawerData = RegisteredDrawers.FindByKey(DrawerId);
+
+		if(DrawerData)
+		{
+			TSharedRef<SWidgetDrawer> ThisStatusBar = SharedThis(this);
+
+			TSharedPtr<SWindow> MyWindow = FSlateApplication::Get().FindWidgetWindow(AsShared());
+
+			const float MaxDrawerHeight = MyWindow->GetSizeInScreen().Y * 0.90f;
+
+			float TargetDrawerHeightPct = .33f;
+			GConfig->GetFloat(TEXT("DrawerSizes"), *(GetSerializableName()+TEXT(".")+DrawerData->UniqueId.ToString()), TargetDrawerHeightPct, GEditorSettingsIni);
+
+			float TargetDrawerHeight = (MyWindow->GetSizeInScreen().Y * TargetDrawerHeightPct) / MyWindow->GetDPIScaleFactor();
+
+			const float MinDrawerHeight = GetTickSpaceGeometry().GetLocalSize().Y + MyWindow->GetWindowBorderSize().Bottom;
+	
+			FOpenDrawerData NewlyOpenedDrawer;
+
+			MyWindow->AddOverlaySlot()
+				.VAlign(VAlign_Bottom)
+				.Padding(FMargin(10.0f, 20.0f, 10.0f, MinDrawerHeight))
+				[
+					SAssignNew(NewlyOpenedDrawer.DrawerOverlay, SDrawerOverlay)
+					.MinDrawerHeight(MinDrawerHeight)
+					.TargetDrawerHeight(TargetDrawerHeight)
+					.MaxDrawerHeight(MaxDrawerHeight)
+					.OnDismissComplete_Lambda(
+						[DrawerId, this]()
+						{
+							CloseDrawerImmediately(DrawerId);
+						})
+					.OnTargetHeightChanged(this, &SWidgetDrawer::OnDrawerHeightChanged)
+					[
+						DrawerData->GetDrawerContentDelegate.Execute()
+					]
+				];
+
+			NewlyOpenedDrawer.WindowWithOverlayContent = MyWindow;
+			NewlyOpenedDrawer.DrawerId = DrawerId;
+			NewlyOpenedDrawer.DrawerOverlay->Open();
+
+			OpenedDrawer = MoveTemp(NewlyOpenedDrawer);
+
+			DrawerData->OnDrawerOpenedDelegate.ExecuteIfBound(ThisStatusBar->DrawerName);
+		}
+	}
+}
+
+bool SWidgetDrawer::DismissDrawer(const TSharedPtr<SWidget>& NewlyFocusedWidget)
+{
+	bool bWasDismissed = false;
+	if (OpenedDrawer.IsValid())
+	{
+		FWidgetDrawerConfig* Drawer = RegisteredDrawers.FindByKey(OpenedDrawer.DrawerId);
+
+		OpenedDrawer.DrawerOverlay->Dismiss();
+		DismissingDrawers.Add(MoveTemp(OpenedDrawer));
+
+		OpenedDrawer = FOpenDrawerData();
+
+		Drawer->OnDrawerDismissedDelegate.ExecuteIfBound(NewlyFocusedWidget);
+		bWasDismissed = true;
+	}
+
+	return bWasDismissed;
+}
+
+void SWidgetDrawer::CloseDrawerImmediately(FName DrawerId)
+{
+	// If no ID is specified remove all drawers
+	if (DrawerId.IsNone())
+	{
+		for (const FOpenDrawerData& Data : DismissingDrawers)
+		{
+			CloseDrawerImmediatelyInternal(Data);
+		}
+
+		DismissingDrawers.Empty();
+
+		CloseDrawerImmediatelyInternal(OpenedDrawer);
+
+		OpenedDrawer = FOpenDrawerData();
+	}
+	else
+	{
+		int32 Index = DismissingDrawers.IndexOfByKey(DrawerId);
+		if (Index != INDEX_NONE)
+		{
+			CloseDrawerImmediatelyInternal(DismissingDrawers[Index]);
+			DismissingDrawers.RemoveAtSwap(Index);
+		}
+		else if (OpenedDrawer == DrawerId)
+		{
+			CloseDrawerImmediatelyInternal(OpenedDrawer);
+			OpenedDrawer = FOpenDrawerData();
+		}
+	}
+}
+
+#undef LOCTEXT_NAMESPACE
\ No newline at end of file
diff --git a/Engine/Source/Developer/ToolWidgets/Public/SWidgetDrawer.h b/Engine/Source/Developer/ToolWidgets/Public/SWidgetDrawer.h
new file mode 100644
--- /dev/null
+++ b/Engine/Source/Developer/ToolWidgets/Public/SWidgetDrawer.h
@@ -0,0 +1,112 @@
+// Copyright Epic Games, Inc. All Rights Reserved.
+
+#pragma once
+
+#include "CoreMinimal.h"
+#include "WidgetDrawerConfig.h"
+#include "Widgets/SCompoundWidget.h"
+#include "Widgets/DeclarativeSyntaxSupport.h"
+
+DECLARE_DELEGATE_OneParam(FOnStatusBarDrawerTargetHeightChanged, float);
+
+class SWindow;
+class SDockTab;
+class SWidget;
+class SDrawerOverlay;
+class SHorizontalBox;
+
+struct TOOLWIDGETS_API FOpenDrawerData
+{
+	FName DrawerId;
+	TSharedPtr<SDrawerOverlay> DrawerOverlay;
+	TWeakPtr<SWindow> WindowWithOverlayContent;
+
+	bool IsValid() const
+	{
+		return !DrawerId.IsNone();
+	}
+
+	bool operator==(const FName InDrawerId) const
+	{
+		return DrawerId == InDrawerId;
+	}
+};
+
+class TOOLWIDGETS_API SWidgetDrawer : public SCompoundWidget
+{
+	SLATE_BEGIN_ARGS(SWidgetDrawer)	
+	{}
+	SLATE_END_ARGS()
+public:
+	
+	~SWidgetDrawer();
+
+	/**
+	 * @param InArgs
+	 * @param InStatusBarName Unique ID used to for writing to config file. Config file saves target drawer height. 
+	 */
+	void Construct(const FArguments& InArgs, FName InStatusBarName);
+	
+	virtual bool SupportsKeyboardFocus() const { return false; }
+	
+
+	/**
+	 * Registers a new drawer with this status bar. Registering will add a button to open and close the drawer
+	 */
+	void RegisterDrawer(FWidgetDrawerConfig&& Drawer, int32 SlotIndex = INDEX_NONE);
+
+	/**
+	 * Opens a drawer
+	 * 
+	 * @param DrawerId Name of the registered drawer to open 
+	 */
+	void OpenDrawer(const FName DrawerId);
+
+	/**
+	 * Dismisses an open drawer with an animation.  The drawer contents are removed once the animation is complete
+	 * 
+	 * @param NewlyFocusedWidget Optional widget to focus 
+	 * @return true if any open drawer was dismissed
+	 */
+	bool DismissDrawer(const TSharedPtr<SWidget>& NewlyFocusedWidget);
+
+	/**
+	 * Closes a drawer immediately with no closing animation. Needed for when UI is shutting down or to prevent conflicts where a drawer is open in two places at once. 
+	 * @param DrawerId The name of the drawer to close. If no name is specified all drawers are closed immediately
+	 */
+	void CloseDrawerImmediately(FName DrawerId = NAME_None);
+
+	/** Is a specific registered drawer currently open */
+	bool IsDrawerOpened(const FName DrawerId) const;
+
+	/** Is any drawer other than the one specified by DrawerId opened */
+	bool IsAnyOtherDrawerOpened(const FName DrawerId) const;
+
+	FName GetDrawerName() const { return DrawerName; }
+	FString GetSerializableName() const;
+
+private:
+	
+	/** Called when global focus changes which is used to determine if we should close an opened content browser drawer */
+	void OnGlobalFocusChanging(const FFocusEvent& FocusEvent, const FWeakWidgetPath& OldFocusedWidgetPath, const TSharedPtr<SWidget>& OldFocusedWidget, const FWidgetPath& NewFocusedWidgetPath, const TSharedPtr<SWidget>& NewFocusedWidget);
+	/** Called when active tab changes which is used to determine if we should close an opened content browser drawer */
+	void OnActiveTabChanged(TSharedPtr<SDockTab> PreviouslyActive, TSharedPtr<SDockTab> NewlyActivated);
+	void OnDrawerHeightChanged(float TargetHeight);
+
+	TSharedRef<SWidget> MakeStatusBarDrawerButton(const FWidgetDrawerConfig& Drawer);
+
+	FReply OnDrawerButtonClicked(const FName DrawerId);
+	
+	void CloseDrawerImmediatelyInternal(const FOpenDrawerData& Data);
+	
+private:
+
+	TArray<FWidgetDrawerConfig> RegisteredDrawers;
+	TArray<FOpenDrawerData> DismissingDrawers;
+	
+	TSharedPtr<SHorizontalBox> DrawerBox;
+
+	FName DrawerName;
+	FOpenDrawerData OpenedDrawer;
+};
+
diff --git a/Engine/Source/Developer/ToolWidgets/Public/WidgetDrawerConfig.h b/Engine/Source/Developer/ToolWidgets/Public/WidgetDrawerConfig.h
new file mode 100644
--- /dev/null
+++ b/Engine/Source/Developer/ToolWidgets/Public/WidgetDrawerConfig.h
@@ -0,0 +1,38 @@
+﻿// Copyright Epic Games, Inc. All Rights Reserved.
+
+#pragma once
+
+#include "CoreMinimal.h"
+#include "Framework/SlateDelegates.h"
+
+class SWidget;
+
+DECLARE_DELEGATE_OneParam(FOnStatusBarDrawerOpened, FName StatusBarName)
+DECLARE_DELEGATE_OneParam(FOnStatusBarDrawerDismissed, const TSharedPtr<SWidget>&)
+
+struct FWidgetDrawerConfig
+{
+	FWidgetDrawerConfig(FName InUniqueId)
+		: UniqueId(InUniqueId)
+	{}
+
+	FName UniqueId;
+	FOnGetContent GetDrawerContentDelegate;
+	FOnStatusBarDrawerOpened OnDrawerOpenedDelegate;
+	FOnStatusBarDrawerDismissed OnDrawerDismissedDelegate;
+
+	TSharedPtr<SWidget> CustomWidget;
+	FText ButtonText;
+	FText ToolTipText;
+	const FSlateBrush* Icon;
+
+	bool operator==(const FName& OtherId) const
+	{
+		return UniqueId == OtherId;
+	}
+
+	bool operator==(const FWidgetDrawerConfig& Other) const
+	{
+		return UniqueId == Other.UniqueId;
+	}
+};
\ No newline at end of file
diff --git a/Engine/Source/Editor/StatusBar/Private/SStatusBar.cpp b/Engine/Source/Editor/StatusBar/Private/SStatusBar.cpp
--- a/Engine/Source/Editor/StatusBar/Private/SStatusBar.cpp
+++ b/Engine/Source/Editor/StatusBar/Private/SStatusBar.cpp
@@ -22,6 +22,7 @@
 #include "Framework/Application/SlateApplication.h"
 #include "Widgets/Layout/SSplitter.h"
 #include "InputCoreTypes.h"
+#include "SWidgetDrawer.h"
 #include "Misc/ConfigCacheIni.h"
 #include "Widgets/Notifications/SProgressBar.h"
 #include "Widgets/Notifications/SNotificationList.h"
@@ -571,21 +572,12 @@ SStatusBar::~SStatusBar()
 
 void SStatusBar::Construct(const FArguments& InArgs, FName InStatusBarName, const TSharedRef<SDockTab> InParentTab)
 {
-	StatusBarName = InStatusBarName;
-	StatusBarToolBarName = FName(*(GetStatusBarSerializableName() + ".ToolBar"));
-	
 	ParentTab = InParentTab;
 
 	UpArrow = FAppStyle::Get().GetBrush("StatusBar.ContentBrowserUp");
 	DownArrow = FAppStyle::Get().GetBrush("StatusBar.ContentBrowserDown");
 
 	const FSlateBrush* StatusBarBackground = FAppStyle::Get().GetBrush("Brushes.Panel");
-
-
-	FSlateApplication::Get().OnFocusChanging().AddSP(this, &SStatusBar::OnGlobalFocusChanging);
-	FGlobalTabmanager::Get()->OnActiveTabChanged_Subscribe(FOnActiveTabChanged::FDelegate::CreateSP(this, &SStatusBar::OnActiveTabChanged));
-	FGlobalTabmanager::Get()->OnTabForegrounded_Subscribe(FOnActiveTabChanged::FDelegate::CreateSP(this, &SStatusBar::OnActiveTabChanged));
-	
 	ChildSlot
 	[
 		SNew(SBox)
@@ -595,7 +587,7 @@ void SStatusBar::Construct(const FArguments& InArgs, FName InStatusBarName, cons
 			+SHorizontalBox::Slot()
 			.AutoWidth()
 			[
-				SAssignNew(DrawerBox, SHorizontalBox)
+				SAssignNew(WidgetDrawer, SWidgetDrawer, InStatusBarName)
 			]
 			+SHorizontalBox::Slot()
 			[
@@ -684,10 +676,6 @@ void SStatusBar::StartProgressNotification(FProgressNotificationHandle InHandle,
 			}
 			UpdateProgressStatus();
 		}
-		else
-		{
-			//UE_LOG(LogUnrealEdEngine, Log, TEXT("Progress notification \"%s\" has no work to do so it will not be displayed"), *DisplayText.ToString()))
-		}
 	}
 }
 
@@ -748,82 +736,6 @@ TSharedPtr<SDockTab> SStatusBar::GetParentTab() const
 	return ParentTab.Pin();
 }
 
-void SStatusBar::OnGlobalFocusChanging(const FFocusEvent& FocusEvent, const FWeakWidgetPath& OldFocusedWidgetPath, const TSharedPtr<SWidget>& OldFocusedWidget, const FWidgetPath& NewFocusedWidgetPath, const TSharedPtr<SWidget>& NewFocusedWidget)
-{
-	// Sometimes when dismissing focus can change which will trigger this again
-	static bool bIsRentrant = false;
-
-	if(!bIsRentrant)
-	{
-		TGuardValue<bool> RentrancyGuard(bIsRentrant, true);
-
-		TSharedRef<SWidget> ThisWidget = AsShared();
-
-		TSharedPtr<SWidget> ActiveDrawerOverlayContent;
-		if (OpenedDrawer.IsValid())
-		{
-			ActiveDrawerOverlayContent = OpenedDrawer.DrawerOverlay;
-		}
-
-		bool bShouldDismiss = false;
-
-		// If we aren't focusing any new widgets, act as if the drawer is in the path 
-		bool bDrawerInPath = NewFocusedWidgetPath.ContainsWidget(ActiveDrawerOverlayContent.Get()) 
-			|| NewFocusedWidgetPath.ContainsWidget(this) 
-			|| NewFocusedWidgetPath.Widgets.Num() == 0;
-
-		// Do not close due to slow tasks as those opening send window activation events
-		if (!GIsSlowTask && !bDrawerInPath && !FSlateApplication::Get().GetActiveModalWindow().IsValid() && ActiveDrawerOverlayContent.IsValid())
-		{
-			if (TSharedPtr<SWidget> MenuHost = FSlateApplication::Get().GetMenuHostWidget())
-			{
-				FWidgetPath MenuHostPath;
-
-				// See if the menu being opened is part of the content browser path and if so the menu should not be dismissed
-				FSlateApplication::Get().GeneratePathToWidgetUnchecked(MenuHost.ToSharedRef(), MenuHostPath, EVisibility::All);
-				if (!MenuHostPath.ContainsWidget(ActiveDrawerOverlayContent.Get()))
-				{
-					bShouldDismiss = true;
-				}
-			}
-			else
-			{
-				bShouldDismiss = true;
-			}
-		}
-
-		if (bShouldDismiss)
-		{
-			DismissDrawer(NewFocusedWidget);
-		}
-	}
-}
-
-void SStatusBar::OnActiveTabChanged(TSharedPtr<SDockTab> PreviouslyActive, TSharedPtr<SDockTab> NewlyActivated)
-{
-	bool bShouldRemoveDrawer = false;
-	if (!PreviouslyActive || !NewlyActivated)
-	{
-		// Remove the content browser if there is some invalid state with the tabs
-		bShouldRemoveDrawer = true;
-	}
-	else if(NewlyActivated->GetTabRole() == ETabRole::MajorTab)
-	{
-		// Remove the content browser if a newly activated tab is a major tab
-		bShouldRemoveDrawer = true;
-	}
-	else if (PreviouslyActive->GetTabManagerPtr() != NewlyActivated->GetTabManagerPtr())
-	{
-		// Remove the content browser if we're switching tab managers (indicates a new status bar is becoming active)
-		bShouldRemoveDrawer = true;
-	}
-
-	if (bShouldRemoveDrawer)
-	{
-		CloseDrawerImmediately();
-	}
-}
-
 FText SStatusBar::GetStatusBarMessage() const
 {
 	FText FullMessage;
@@ -840,78 +752,6 @@ FText SStatusBar::GetStatusBarMessage() const
 	return FullMessage;
 }
 
-TSharedRef<SWidget> SStatusBar::MakeStatusBarDrawerButton(const FStatusBarDrawer& Drawer)
-{
-	const FName DrawerId = Drawer.UniqueId;
-
-	const FSlateBrush* StatusBarBackground = FAppStyle::Get().GetBrush("Brushes.Panel");
-
-	TSharedRef<SWidget> DrawerButton = 
-
-		SNew(SBorder)
-		.Padding(FMargin(2.0f, 0.0f))
-		.BorderImage(StatusBarBackground)
-		.Visibility(EVisibility::SelfHitTestInvisible)
-		.VAlign(VAlign_Center)
-		[
-			SNew(SButton)
-			.IsFocusable(false)
-			.ButtonStyle(&FAppStyle::Get().GetWidgetStyle<FButtonStyle>("StatusBar.StatusBarButton"))
-			.OnClicked(this, &SStatusBar::OnDrawerButtonClicked, DrawerId)
-			.ToolTipText(Drawer.ToolTipText)
-			[
-				SNew(SHorizontalBox)
-				+ SHorizontalBox::Slot()
-				.Padding(2.0f)
-				.HAlign(HAlign_Left)
-				.VAlign(VAlign_Center)
-				.AutoWidth()
-				[
-					SNew(SImage)
-					.ColorAndOpacity(FSlateColor::UseForeground())
-					.Image(Drawer.Icon)
-				]
-				+ SHorizontalBox::Slot()
-				.VAlign(VAlign_Center)
-				.Padding(2.0f)
-				[
-					SNew(STextBlock)
-					.TextStyle(&FAppStyle::Get().GetWidgetStyle<FTextBlockStyle>("NormalText"))
-					.Text(Drawer.ButtonText)
-				]
-			]
-		];
-
-
-	if (Drawer.CustomWidget)
-	{
-		return
-			SNew(SHorizontalBox)
-			+ SHorizontalBox::Slot()
-			.Padding(0.0f, 0.0f, 2.0f, 0.0f)
-			.AutoWidth()
-			[
-				DrawerButton
-			]
-			+ SHorizontalBox::Slot()
-			[
-				SNew(SBorder)
-				.Padding(FMargin(2.0f, 0.0f))
-				.BorderImage(StatusBarBackground)
-				.Visibility(EVisibility::SelfHitTestInvisible)
-				.VAlign(VAlign_Center)
-				[
-					Drawer.CustomWidget.ToSharedRef()
-				]
-			];
-	
-	}
-	else
-	{
-		return DrawerButton;
-	}
-}
-
 TSharedRef<SWidget> SStatusBar::MakeStatusBarToolBarWidget()
 {
 	RegisterStatusBarMenu();
@@ -919,7 +759,7 @@ TSharedRef<SWidget> SStatusBar::MakeStatusBarToolBarWidget()
 	FToolMenuContext MenuContext;
 	RegisterSourceControlStatus();
 
-	return UToolMenus::Get()->GenerateWidget(StatusBarToolBarName, MenuContext);
+	return UToolMenus::Get()->GenerateWidget(GetToolbarName(), MenuContext);
 }
 
 TSharedRef<SWidget> SStatusBar::MakeStatusMessageWidget()
@@ -955,19 +795,39 @@ TSharedRef<SWidget> SStatusBar::MakeProgressBar()
 		.OnGetProgressMenuContent(this, &SStatusBar::OnGetProgressBarMenuContent);
 }
 
+void SStatusBar::RegisterDrawer(FWidgetDrawerConfig&& Drawer, int32 SlotIndex)
+{
+	WidgetDrawer->RegisterDrawer(MoveTemp(Drawer), SlotIndex);
+}
+
+void SStatusBar::OpenDrawer(const FName DrawerId)
+{
+	WidgetDrawer->OpenDrawer(DrawerId);
+}
+
+bool SStatusBar::DismissDrawer(const TSharedPtr<SWidget>& NewlyFocusedWidget)
+{
+	return WidgetDrawer->DismissDrawer(NewlyFocusedWidget);
+}
+
+void SStatusBar::CloseDrawerImmediately(FName DrawerId)
+{
+	WidgetDrawer->CloseDrawerImmediately(DrawerId);
+}
+
 bool SStatusBar::IsDrawerOpened(const FName DrawerId) const
 {
-	return OpenedDrawer == DrawerId ? true : false;
+	return WidgetDrawer->IsDrawerOpened(DrawerId);
 }
 
 bool SStatusBar::IsAnyOtherDrawerOpened(const FName DrawerId) const
 {
-	return OpenedDrawer.IsValid() && OpenedDrawer.DrawerId != DrawerId ? true : false;
+	return WidgetDrawer->IsAnyOtherDrawerOpened(DrawerId);
 }
 
 FName SStatusBar::GetStatusBarName() const
 {
-	return StatusBarName;
+	return WidgetDrawer->GetDrawerName();
 }
 
 FReply SStatusBar::OnDrawerButtonClicked(const FName DrawerId)
@@ -984,27 +844,15 @@ FReply SStatusBar::OnDrawerButtonClicked(const FName DrawerId)
 	return FReply::Handled();
 }
 
-
-
-void SStatusBar::OnDrawerHeightChanged(float TargetHeight)
-{
-	TSharedPtr<SWindow> MyWindow = OpenedDrawer.WindowWithOverlayContent.Pin();
-
-	// Save the height has a percentage of the screen
-	const float TargetDrawerHeightPct = TargetHeight / (MyWindow->GetSizeInScreen().Y / MyWindow->GetDPIScaleFactor());
-
-	GConfig->SetFloat(TEXT("DrawerSizes"), *(GetStatusBarSerializableName() + TEXT(".") + OpenedDrawer.DrawerId.ToString()), TargetDrawerHeightPct, GEditorSettingsIni);
-}
-
 void SStatusBar::RegisterStatusBarMenu()
 {
 	UToolMenus* ToolMenus = UToolMenus::Get();
-	if (ToolMenus->IsMenuRegistered(StatusBarToolBarName))
+	if (ToolMenus->IsMenuRegistered(GetToolbarName()))
 	{
 		return;
 	}
 
-	UToolMenu* ToolBar = ToolMenus->RegisterMenu(StatusBarToolBarName, NAME_None, EMultiBoxType::SlimHorizontalToolBar);
+	UToolMenu* ToolBar = ToolMenus->RegisterMenu(GetToolbarName(), NAME_None, EMultiBoxType::SlimHorizontalToolBar);
 	ToolBar->StyleName = "StatusBarToolBar";
 }
 
@@ -1013,7 +861,7 @@ void SStatusBar::RegisterSourceControlStatus()
 	// Source Control preferences
 	FSourceControlMenuHelpers::CheckSourceControlStatus();
 	{
-		UToolMenu* SourceControlMenu = UToolMenus::Get()->ExtendMenu(StatusBarToolBarName);
+		UToolMenu* SourceControlMenu = UToolMenus::Get()->ExtendMenu(GetToolbarName());
 		FToolMenuSection& Section = SourceControlMenu->FindOrAddSection("SourceControl");
 
 		Section.AddEntry(
@@ -1178,146 +1026,9 @@ TSharedRef<SWidget> SStatusBar::OnGetProgressBarMenuContent()
 		];
 }
 
-void SStatusBar::CloseDrawerImmediatelyInternal(const FOpenDrawerData& Data)
-{
-	if (Data.IsValid())
-	{
-		TSharedRef<SWidget> DrawerOverlayContent = Data.DrawerOverlay.ToSharedRef();
-
-		// Remove the content browser from the window
-		if (TSharedPtr<SWindow> Window = Data.WindowWithOverlayContent.Pin())
-		{
-			Window->RemoveOverlaySlot(DrawerOverlayContent);
-		}
-	}
-}
-
-FString SStatusBar::GetStatusBarSerializableName() const
-{
-	return StatusBarName.GetPlainNameString();
-}
-
-void SStatusBar::RegisterDrawer(FStatusBarDrawer&& Drawer, int32 SlotIndex)
-{
-	const int32 NumDrawers = RegisteredDrawers.Num();
-	RegisteredDrawers.AddUnique(Drawer);
-
-	if (RegisteredDrawers.Num() > NumDrawers)
-	{
-		DrawerBox->InsertSlot(SlotIndex)
-		.Padding(1.0f, 0.0f)
-		.AutoWidth()
-		[
-			MakeStatusBarDrawerButton(Drawer)
-		];
-	}
-}
-
-void SStatusBar::OpenDrawer(const FName DrawerId)
+FName SStatusBar::GetToolbarName() const
 {
-	// Close any other open drawer
-	if (OpenedDrawer.DrawerId != DrawerId && DismissingDrawers.IndexOfByKey(DrawerId) == INDEX_NONE)
-	{
-		DismissDrawer(nullptr);
-
-		FStatusBarDrawer* DrawerData = RegisteredDrawers.FindByKey(DrawerId);
-
-		if(DrawerData)
-		{
-			TSharedRef<SStatusBar> ThisStatusBar = SharedThis(this);
-
-			TSharedPtr<SWindow> MyWindow = FSlateApplication::Get().FindWidgetWindow(AsShared());
-
-			const float MaxDrawerHeight = MyWindow->GetSizeInScreen().Y * 0.90f;
-
-			float TargetDrawerHeightPct = .33f;
-			GConfig->GetFloat(TEXT("DrawerSizes"), *(GetStatusBarSerializableName()+TEXT(".")+DrawerData->UniqueId.ToString()), TargetDrawerHeightPct, GEditorSettingsIni);
-
-			float TargetDrawerHeight = (MyWindow->GetSizeInScreen().Y * TargetDrawerHeightPct) / MyWindow->GetDPIScaleFactor();
-
-			const float MinDrawerHeight = GetTickSpaceGeometry().GetLocalSize().Y + MyWindow->GetWindowBorderSize().Bottom;
-	
-			FOpenDrawerData NewlyOpenedDrawer;
-
-			MyWindow->AddOverlaySlot()
-				.VAlign(VAlign_Bottom)
-				.Padding(FMargin(10.0f, 20.0f, 10.0f, MinDrawerHeight))
-				[
-					SAssignNew(NewlyOpenedDrawer.DrawerOverlay, SDrawerOverlay)
-					.MinDrawerHeight(MinDrawerHeight)
-					.TargetDrawerHeight(TargetDrawerHeight)
-					.MaxDrawerHeight(MaxDrawerHeight)
-					.OnDismissComplete_Lambda(
-						[DrawerId, this]()
-						{
-							CloseDrawerImmediately(DrawerId);
-						})
-					.OnTargetHeightChanged(this, &SStatusBar::OnDrawerHeightChanged)
-					[
-						DrawerData->GetDrawerContentDelegate.Execute()
-					]
-				];
-
-			NewlyOpenedDrawer.WindowWithOverlayContent = MyWindow;
-			NewlyOpenedDrawer.DrawerId = DrawerId;
-			NewlyOpenedDrawer.DrawerOverlay->Open();
-
-			OpenedDrawer = MoveTemp(NewlyOpenedDrawer);
-
-			DrawerData->OnDrawerOpenedDelegate.ExecuteIfBound(ThisStatusBar->StatusBarName);
-		}
-	}
-}
-
-bool SStatusBar::DismissDrawer(const TSharedPtr<SWidget>& NewlyFocusedWidget)
-{
-	bool bWasDismissed = false;
-	if (OpenedDrawer.IsValid())
-	{
-		FStatusBarDrawer* Drawer = RegisteredDrawers.FindByKey(OpenedDrawer.DrawerId);
-
-		OpenedDrawer.DrawerOverlay->Dismiss();
-		DismissingDrawers.Add(MoveTemp(OpenedDrawer));
-
-		OpenedDrawer = FOpenDrawerData();
-
-		Drawer->OnDrawerDismissedDelegate.ExecuteIfBound(NewlyFocusedWidget);
-		bWasDismissed = true;
-	}
-
-	return bWasDismissed;
-}
-
-void SStatusBar::CloseDrawerImmediately(FName DrawerId)
-{
-	// If no ID is specified remove all drawers
-	if (DrawerId.IsNone())
-	{
-		for (const FOpenDrawerData& Data : DismissingDrawers)
-		{
-			CloseDrawerImmediatelyInternal(Data);
-		}
-
-		DismissingDrawers.Empty();
-
-		CloseDrawerImmediatelyInternal(OpenedDrawer);
-
-		OpenedDrawer = FOpenDrawerData();
-	}
-	else
-	{
-		int32 Index = DismissingDrawers.IndexOfByKey(DrawerId);
-		if (Index != INDEX_NONE)
-		{
-			CloseDrawerImmediatelyInternal(DismissingDrawers[Index]);
-			DismissingDrawers.RemoveAtSwap(Index);
-		}
-		else if (OpenedDrawer == DrawerId)
-		{
-			CloseDrawerImmediatelyInternal(OpenedDrawer);
-			OpenedDrawer = FOpenDrawerData();
-		}
-	}
+	return FName(*(WidgetDrawer->GetSerializableName() + ".ToolBar"));
 }
 
 #undef LOCTEXT_NAMESPACE
\ No newline at end of file
diff --git a/Engine/Source/Editor/StatusBar/Private/SStatusBar.h b/Engine/Source/Editor/StatusBar/Private/SStatusBar.h
--- a/Engine/Source/Editor/StatusBar/Private/SStatusBar.h
+++ b/Engine/Source/Editor/StatusBar/Private/SStatusBar.h
@@ -5,9 +5,6 @@
 #include "StatusBarSubsystem.h"
 #include "Widgets/SCompoundWidget.h"
 #include "Widgets/DeclarativeSyntaxSupport.h"
-#include "Framework/SlateDelegates.h"
-#include "Animation/CurveSequence.h"
-#include "Application/ThrottleManager.h"
 
 DECLARE_DELEGATE_OneParam(FOnStatusBarDrawerTargetHeightChanged, float);
 
@@ -19,6 +16,7 @@ class SDrawerOverlay;
 class SHorizontalBox;
 class SStatusBarProgressArea;
 class SStatusBarProgressWidget;
+class SWidgetDrawer;
 
 namespace StatusBarDrawerIds
 {
@@ -58,23 +56,6 @@ struct FStatusBarProgress
 	FProgressNotificationHandle Handle;
 };
 
-struct FOpenDrawerData
-{
-	FName DrawerId;
-	TSharedPtr<SDrawerOverlay> DrawerOverlay;
-	TWeakPtr<SWindow> WindowWithOverlayContent;
-
-	bool IsValid() const
-	{
-		return !DrawerId.IsNone();
-	}
-
-	bool operator==(const FName InDrawerId) const
-	{
-		return DrawerId == InDrawerId;
-	}
-};
-
 class SStatusBar : public SCompoundWidget
 {
 	SLATE_BEGIN_ARGS(SStatusBar)	
@@ -139,7 +120,7 @@ public:
 	/**
 	 * Registers a new drawer with this status bar. Registering will add a button to open and close the drawer
 	 */
-	void RegisterDrawer(FStatusBarDrawer&& Drawer, int32 SlotIndex = INDEX_NONE);
+	void RegisterDrawer(FWidgetDrawerConfig&& Drawer, int32 SlotIndex = INDEX_NONE);
 
 	/**
 	 * Opens a drawer
@@ -172,19 +153,11 @@ public:
 	FName GetStatusBarName() const;
 
 private:
-	/** Called when global focus changes which is used to determine if we should close an opened content browser drawer */
-	void OnGlobalFocusChanging(const FFocusEvent& FocusEvent, const FWeakWidgetPath& OldFocusedWidgetPath, const TSharedPtr<SWidget>& OldFocusedWidget, const FWidgetPath& NewFocusedWidgetPath, const TSharedPtr<SWidget>& NewFocusedWidget);
-
-	/** Called when active tab changes which is used to determine if we should close an opened content browser drawer */
-	void OnActiveTabChanged(TSharedPtr<SDockTab> PreviouslyActive, TSharedPtr<SDockTab> NewlyActivated);
-
-	void OnDrawerHeightChanged(float TargetHeight);
 
 	EVisibility GetHelpIconVisibility() const;
 
 	FText GetStatusBarMessage() const;
 
-	TSharedRef<SWidget> MakeStatusBarDrawerButton(const FStatusBarDrawer& Drawer);
 	TSharedRef<SWidget> MakeStatusBarToolBarWidget();
 	TSharedRef<SWidget> MakeStatusMessageWidget();
 	TSharedRef<SWidget> MakeProgressBar();
@@ -192,7 +165,6 @@ private:
 	FReply OnDrawerButtonClicked(const FName DrawerId);
 
 	void RegisterStatusBarMenu();
-
 	void RegisterSourceControlStatus();
 
 	FStatusBarProgress* FindProgressNotification(FProgressNotificationHandle InHandle);
@@ -203,35 +175,22 @@ private:
 	
 	TSharedRef<SWidget> OnGetProgressBarMenuContent();
 
-	void CloseDrawerImmediatelyInternal(const FOpenDrawerData& Data);
-
-	FString GetStatusBarSerializableName() const;
+	FName GetToolbarName() const;
+	
 private:
+	
 	TArray<FStatusBarMessage> MessageStack;
 	TArray<FStatusBarProgress> ProgressNotifications;
 
 	TWeakPtr<SDockTab> ParentTab;
-
-	TArray<FStatusBarDrawer> RegisteredDrawers;
-
-	FOpenDrawerData OpenedDrawer;
-
-	TArray<FOpenDrawerData> DismissingDrawers;
-
-	TSharedPtr<SHorizontalBox> DrawerBox;
-
+	TSharedPtr<SWidgetDrawer> WidgetDrawer;
 	TSharedPtr<SStatusBarProgressArea> ProgressBar;
-
 	TWeakPtr<SNotificationItem> ActiveProgressNotification;
-
 	TWeakPtr<SStatusBarProgressWidget> ActiveNotificationProgressWidget;
 
 	const FSlateBrush* UpArrow = nullptr;
 	const FSlateBrush* DownArrow = nullptr;
-	FName StatusBarName;
-	FName StatusBarToolBarName;
 
 	bool bAllowedToRefreshProgressNotification = false;
-
 };
 
diff --git a/Engine/Source/Editor/StatusBar/Private/StatusBarSubsystem.cpp b/Engine/Source/Editor/StatusBar/Private/StatusBarSubsystem.cpp
--- a/Engine/Source/Editor/StatusBar/Private/StatusBarSubsystem.cpp
+++ b/Engine/Source/Editor/StatusBar/Private/StatusBarSubsystem.cpp
@@ -22,6 +22,7 @@
 #include "Widgets/Input/SButton.h"
 #include "Misc/ConfigCacheIni.h"
 #include "OutputLogModule.h"
+#include "WidgetDrawerConfig.h"
 #include "Widgets/Input/SMultiLineEditableTextBox.h"
 #include "Classes/EditorStyleSettings.h"
 
@@ -433,7 +434,7 @@ TSharedRef<SWidget> UStatusBarSubsystem::MakeStatusBarWidget(FName StatusBarName
 	TSharedRef<SStatusBar> StatusBar =
 		SNew(SStatusBar, StatusBarName, InParentTab);
 
-	FStatusBarDrawer ContentBrowserDrawer(StatusBarDrawerIds::ContentBrowser);
+	FWidgetDrawerConfig ContentBrowserDrawer(StatusBarDrawerIds::ContentBrowser);
 	ContentBrowserDrawer.GetDrawerContentDelegate.BindUObject(this, &UStatusBarSubsystem::OnGetContentBrowser);
 	ContentBrowserDrawer.OnDrawerOpenedDelegate.BindUObject(this, &UStatusBarSubsystem::OnContentBrowserOpened);
 	ContentBrowserDrawer.OnDrawerDismissedDelegate.BindUObject(this, &UStatusBarSubsystem::OnContentBrowserDismissed);
@@ -465,7 +466,7 @@ TSharedRef<SWidget> UStatusBarSubsystem::MakeStatusBarWidget(FName StatusBarName
 			];
 	
 
-	FStatusBarDrawer OutputLogDrawer(StatusBarDrawerIds::OutputLog);
+	FWidgetDrawerConfig OutputLogDrawer(StatusBarDrawerIds::OutputLog);
 
 	OutputLogDrawer.GetDrawerContentDelegate.BindUObject(this, &UStatusBarSubsystem::OnGetOutputLog);
 	OutputLogDrawer.OnDrawerOpenedDelegate.BindUObject(this, &UStatusBarSubsystem::OnOutputLogOpened);
@@ -517,7 +518,7 @@ bool UStatusBarSubsystem::ActiveWindowHasStatusBar() const
 	return false;
 }
 
-void UStatusBarSubsystem::RegisterDrawer(FName StatusBarName, FStatusBarDrawer&& Drawer, int32 SlotIndex)
+void UStatusBarSubsystem::RegisterDrawer(FName StatusBarName, FWidgetDrawerConfig&& Drawer, int32 SlotIndex)
 {
 	if (TSharedPtr<SStatusBar> StatusBar = GetStatusBar(StatusBarName))
 	{
diff --git a/Engine/Source/Editor/StatusBar/Public/StatusBarSubsystem.h b/Engine/Source/Editor/StatusBar/Public/StatusBarSubsystem.h
--- a/Engine/Source/Editor/StatusBar/Public/StatusBarSubsystem.h
+++ b/Engine/Source/Editor/StatusBar/Public/StatusBarSubsystem.h
@@ -5,17 +5,13 @@
 #include "CoreMinimal.h"
 #include "EditorSubsystem.h"
 #include "Framework/Notifications/NotificationManager.h"
-#include "Framework/SlateDelegates.h"
-
 #include "StatusBarSubsystem.generated.h"
 
-DECLARE_DELEGATE_OneParam(FOnStatusBarDrawerOpened, FName StatusBarName)
-DECLARE_DELEGATE_OneParam(FOnStatusBarDrawerDismissed, const TSharedPtr<SWidget>&)
-
 class SStatusBar;
 class SWindow;
 class SWidget;
 class SDockTab;
+struct FWidgetDrawerConfig;
 
 template<typename ObjectType> 
 class TAttribute;
@@ -50,33 +46,6 @@ private:
 	int32 Id;
 };
 
-struct FStatusBarDrawer
-{
-	FStatusBarDrawer(FName InUniqueId)
-		: UniqueId(InUniqueId)
-	{}
-
-	FName UniqueId;
-	FOnGetContent GetDrawerContentDelegate;
-	FOnStatusBarDrawerOpened OnDrawerOpenedDelegate;
-	FOnStatusBarDrawerDismissed OnDrawerDismissedDelegate;
-
-	TSharedPtr<SWidget> CustomWidget;
-	FText ButtonText;
-	FText ToolTipText;
-	const FSlateBrush* Icon;
-
-	bool operator==(const FName& OtherId) const
-	{
-		return UniqueId == OtherId;
-	}
-
-	bool operator==(const FStatusBarDrawer& Other) const
-	{
-		return UniqueId == Other.UniqueId;
-	}
-};
-
 struct FStatusBarData
 {
 	TWeakPtr<SStatusBar> StatusBarWidget;
@@ -153,7 +122,7 @@ public:
 	 * @param Drawer		The drawer to add to the status bar
 	 * @param SlotIndex		The position at which to add the new drawer
 	 */
-	void RegisterDrawer(FName StatusBarName, FStatusBarDrawer&& Drawer, int32 SlotIndex = INDEX_NONE);
+	void RegisterDrawer(FName StatusBarName, FWidgetDrawerConfig&& Drawer, int32 SlotIndex = INDEX_NONE);
 
 	/** 
 	 * Pushes a new status bar message
diff --git a/Engine/Source/Editor/StatusBar/StatusBar.Build.cs b/Engine/Source/Editor/StatusBar/StatusBar.Build.cs
--- a/Engine/Source/Editor/StatusBar/StatusBar.Build.cs
+++ b/Engine/Source/Editor/StatusBar/StatusBar.Build.cs
@@ -13,17 +13,18 @@ public class StatusBar : ModuleRules
 			new string[] {
 				"Core",
 				"CoreUObject",
+				"EditorFramework",
+				"EditorStyle",
+				"EditorSubsystem",
 				"Engine",
-				"Slate",
-				"SlateCore",
 				"InputCore",
-				"EditorStyle",
-				"EditorFramework",
-				"UnrealEd",
-				"ToolMenus",
 				"OutputLog",
+				"Slate",
+				"SlateCore",
 				"SourceControlWindows",
-				"EditorSubsystem",
+				"ToolMenus",
+				"ToolWidgets",
+				"UnrealEd",
 			});
 
 		DynamicallyLoadedModuleNames.AddRange(
@@ -31,7 +32,6 @@ public class StatusBar : ModuleRules
 				"MainFrame",
 			});
 
-
 		PrivateIncludePathModuleNames.AddRange(
 			new string[] {
 				"MainFrame",
diff --git a/Engine/Source/Editor/UMGEditor/Private/BlueprintModes/WidgetDesignerApplicationMode.cpp b/Engine/Source/Editor/UMGEditor/Private/BlueprintModes/WidgetDesignerApplicationMode.cpp
--- a/Engine/Source/Editor/UMGEditor/Private/BlueprintModes/WidgetDesignerApplicationMode.cpp
+++ b/Engine/Source/Editor/UMGEditor/Private/BlueprintModes/WidgetDesignerApplicationMode.cpp
@@ -10,6 +10,7 @@
 #include "UMGEditorModule.h"
 #include "StatusBarSubsystem.h"
 #include "ToolMenus.h"
+#include "WidgetDrawerConfig.h"
 
 #include "TabFactory/PaletteTabSummoner.h"
 #include "TabFactory/LibraryTabSummoner.h"
@@ -142,7 +143,7 @@ void FWidgetDesignerApplicationMode::PostActivateMode()
 
 	TSharedPtr<FWidgetBlueprintEditor> BP = GetBlueprintEditor();
 
-	FStatusBarDrawer WidgetAnimSequencerDrawer(FAnimationTabSummoner::WidgetAnimSequencerDrawerID);
+	FWidgetDrawerConfig WidgetAnimSequencerDrawer(FAnimationTabSummoner::WidgetAnimSequencerDrawerID);
 	WidgetAnimSequencerDrawer.GetDrawerContentDelegate.BindSP(BP.Get(), &FWidgetBlueprintEditor::OnGetWidgetAnimSequencer);
 	WidgetAnimSequencerDrawer.OnDrawerOpenedDelegate.BindSP(BP.Get(), &FWidgetBlueprintEditor::OnWidgetAnimSequencerOpened);
 	WidgetAnimSequencerDrawer.OnDrawerDismissedDelegate.BindSP(BP.Get(), &FWidgetBlueprintEditor::OnWidgetAnimSequencerDismissed);
diff --git a/Engine/Source/Editor/UnrealEd/Private/Toolkits/AssetEditorToolkit.cpp b/Engine/Source/Editor/UnrealEd/Private/Toolkits/AssetEditorToolkit.cpp
--- a/Engine/Source/Editor/UnrealEd/Private/Toolkits/AssetEditorToolkit.cpp
+++ b/Engine/Source/Editor/UnrealEd/Private/Toolkits/AssetEditorToolkit.cpp
@@ -33,6 +33,7 @@
 #include "AssetEditorModeManager.h"
 #include "Misc/Attribute.h"
 #include "Textures/SlateIcon.h"
+#include "WidgetDrawerConfig.h"
 
 #define LOCTEXT_NAMESPACE "AssetEditorToolkit"
 
@@ -1159,7 +1160,7 @@ void FAssetEditorToolkit::RegenerateMenusAndToolbars()
 }
 
 
-void FAssetEditorToolkit::RegisterDrawer(FStatusBarDrawer&& Drawer, int32 SlotIndex)
+void FAssetEditorToolkit::RegisterDrawer(FWidgetDrawerConfig&& Drawer, int32 SlotIndex)
 {
 	TSharedPtr< class SStandaloneAssetEditorToolkitHost > HostWidget = StandaloneHost.Pin();
 	if (HostWidget.IsValid())
diff --git a/Engine/Source/Editor/UnrealEd/Private/Toolkits/SStandaloneAssetEditorToolkitHost.cpp b/Engine/Source/Editor/UnrealEd/Private/Toolkits/SStandaloneAssetEditorToolkitHost.cpp
--- a/Engine/Source/Editor/UnrealEd/Private/Toolkits/SStandaloneAssetEditorToolkitHost.cpp
+++ b/Engine/Source/Editor/UnrealEd/Private/Toolkits/SStandaloneAssetEditorToolkitHost.cpp
@@ -16,6 +16,7 @@
 #include "ToolMenus.h"
 #include "EditorModeManager.h"
 #include "EditorModes.h"
+#include "WidgetDrawerConfig.h"
 
 #define LOCTEXT_NAMESPACE "StandaloneAssetEditorToolkit"
 
@@ -242,7 +243,7 @@ void SStandaloneAssetEditorToolkitHost::SetToolbar(TSharedPtr<SWidget> Toolbar)
 	}
 }
 
-void SStandaloneAssetEditorToolkitHost::RegisterDrawer(FStatusBarDrawer&& Drawer, int32 SlotIndex)
+void SStandaloneAssetEditorToolkitHost::RegisterDrawer(FWidgetDrawerConfig&& Drawer, int32 SlotIndex)
 {
 	if (StatusBarWidget.IsValid())
 	{
diff --git a/Engine/Source/Editor/UnrealEd/Private/Toolkits/SStandaloneAssetEditorToolkitHost.h b/Engine/Source/Editor/UnrealEd/Private/Toolkits/SStandaloneAssetEditorToolkitHost.h
--- a/Engine/Source/Editor/UnrealEd/Private/Toolkits/SStandaloneAssetEditorToolkitHost.h
+++ b/Engine/Source/Editor/UnrealEd/Private/Toolkits/SStandaloneAssetEditorToolkitHost.h
@@ -97,7 +97,7 @@ public:
 	/**
 	 * Registers a drawer for the asset editor status bar
 	 */
-	void RegisterDrawer(struct FStatusBarDrawer&& Drawer, int32 SlotIndex = INDEX_NONE);
+	void RegisterDrawer(struct FWidgetDrawerConfig&& Drawer, int32 SlotIndex = INDEX_NONE);
 
 	virtual FEditorModeTools& GetEditorModeManager() const override;
 
diff --git a/Engine/Source/Editor/UnrealEd/Public/Toolkits/AssetEditorToolkit.h b/Engine/Source/Editor/UnrealEd/Public/Toolkits/AssetEditorToolkit.h
--- a/Engine/Source/Editor/UnrealEd/Public/Toolkits/AssetEditorToolkit.h
+++ b/Engine/Source/Editor/UnrealEd/Public/Toolkits/AssetEditorToolkit.h
@@ -180,7 +180,7 @@ public:
 	/**
 	 * Registers a drawer for the asset editor status bar
 	 */
-	void RegisterDrawer(struct FStatusBarDrawer&& Drawer, int32 SlotIndex = INDEX_NONE);
+	void RegisterDrawer(struct FWidgetDrawerConfig&& Drawer, int32 SlotIndex = INDEX_NONE);
 
 	/** Applies the passed in layout (or the saved user-modified version if available).  Must be called after InitAssetEditor. */
 	void RestoreFromLayout(const TSharedRef<FTabManager::FLayout>& NewLayout);
