Fix non-unity issue

diff --git a/Engine/Source/Developer/ToolWidgets/Private/SWidgetDrawer.cpp b/Engine/Source/Developer/ToolWidgets/Private/SWidgetDrawer.cpp
--- a/Engine/Source/Developer/ToolWidgets/Private/SWidgetDrawer.cpp
+++ b/Engine/Source/Developer/ToolWidgets/Private/SWidgetDrawer.cpp
@@ -5,6 +5,7 @@
 #include "Widgets/Input/SButton.h"
 #include "Widgets/Layout/SBorder.h"
 #include "Widgets/Text/SRichTextBlock.h"
+#include "Widgets/Text/STextBlock.h"
 #include "Widgets/Docking/SDockTab.h"
 #include "Framework/Application/SlateApplication.h"
 #include "InputCoreTypes.h"
