Move SCustomDialog from UnrealEd to ToolWidgets

Unshelved from pending changelist '19221005':

diff --git a/Engine/Plugins/Animation/ControlRig/Source/ControlRigEditor/ControlRigEditor.Build.cs b/Engine/Plugins/Animation/ControlRig/Source/ControlRigEditor/ControlRigEditor.Build.cs
--- a/Engine/Plugins/Animation/ControlRig/Source/ControlRigEditor/ControlRigEditor.Build.cs
+++ b/Engine/Plugins/Animation/ControlRig/Source/ControlRigEditor/ControlRigEditor.Build.cs
@@ -85,6 +85,7 @@ namespace UnrealBuildTool.Rules
 					"AdvancedPreviewScene",
 					"ToolWidgets",
                     "AnimationWidgets",
+                    "ToolWidgets"
 				}
             );
 
diff --git a/Engine/Plugins/Animation/ControlRig/Source/ControlRigEditor/Private/ControlRigEditorModule.cpp b/Engine/Plugins/Animation/ControlRig/Source/ControlRigEditor/Private/ControlRigEditorModule.cpp
--- a/Engine/Plugins/Animation/ControlRig/Source/ControlRigEditor/Private/ControlRigEditorModule.cpp
+++ b/Engine/Plugins/Animation/ControlRig/Source/ControlRigEditor/Private/ControlRigEditorModule.cpp
@@ -113,7 +113,7 @@
 #include "AssetTypeActions_ControlRigPose.h"
 #include "ControlRigBlueprintFactory.h"
 #include "ControlRigPythonLogDetails.h"
-#include "Dialogs/CustomDialog.h"
+#include "Dialog/SCustomDialog.h"
 #include "Sequencer/MovieSceneControlRigSpaceChannel.h"
 #include "SequencerChannelInterface.h"
 #include "Framework/Notifications/NotificationManager.h"
@@ -1895,7 +1895,10 @@ void FControlRigEditorModule::GetContextMenuActions(const UControlRigGraphSchema
 
 							TSharedRef<SCustomDialog> MirrorDialog = SNew(SCustomDialog)
 								.Title(FText(LOCTEXT("ControlRigHierarchyMirror", "Mirror Graph")))
-								.DialogContent( KismetInspector)
+								.Content()
+								[
+									KismetInspector
+								]
 								.Buttons({
 									SCustomDialog::FButton(LOCTEXT("OK", "OK")),
 									SCustomDialog::FButton(LOCTEXT("Cancel", "Cancel"))
diff --git a/Engine/Plugins/Animation/ControlRig/Source/ControlRigEditor/Private/Editor/SRigHierarchy.cpp b/Engine/Plugins/Animation/ControlRig/Source/ControlRigEditor/Private/Editor/SRigHierarchy.cpp
--- a/Engine/Plugins/Animation/ControlRig/Source/ControlRigEditor/Private/Editor/SRigHierarchy.cpp
+++ b/Engine/Plugins/Animation/ControlRig/Source/ControlRigEditor/Private/Editor/SRigHierarchy.cpp
@@ -36,7 +36,7 @@
 #include "IPersonaToolkit.h"
 #include "SKismetInspector.h"
 #include "Types/WidgetActiveTimerDelegate.h"
-#include "Dialogs/CustomDialog.h"
+#include "Dialog/SCustomDialog.h"
 #include "EditMode/ControlRigEditMode.h"
 #include "ToolMenus.h"
 #include "ControlRigContextMenuContext.h"
@@ -1677,7 +1677,10 @@ void SRigHierarchy::HandleMirrorItem()
 
 		TSharedRef<SCustomDialog> MirrorDialog = SNew(SCustomDialog)
 			.Title(FText(LOCTEXT("ControlRigHierarchyMirror", "Mirror Selected Rig Elements")))
-			.DialogContent( KismetInspector)
+			.Content()
+			[
+				KismetInspector
+			]
 			.Buttons({
 				SCustomDialog::FButton(LOCTEXT("OK", "OK")),
 				SCustomDialog::FButton(LOCTEXT("Cancel", "Cancel"))
diff --git a/Engine/Plugins/Animation/IKRig/Source/IKRigEditor/Private/RigEditor/IKRigEditorController.cpp b/Engine/Plugins/Animation/IKRig/Source/IKRigEditor/Private/RigEditor/IKRigEditorController.cpp
--- a/Engine/Plugins/Animation/IKRig/Source/IKRigEditor/Private/RigEditor/IKRigEditorController.cpp
+++ b/Engine/Plugins/Animation/IKRig/Source/IKRigEditor/Private/RigEditor/IKRigEditorController.cpp
@@ -4,7 +4,7 @@
 
 #include "IPersonaToolkit.h"
 #include "SKismetInspector.h"
-#include "Dialogs/CustomDialog.h"
+#include "Dialog/SCustomDialog.h"
 #include "Dialogs/Dialogs.h"
 
 #include "RigEditor/IKRigController.h"
@@ -474,7 +474,10 @@ bool FIKRigEditorController::PromptToAddSolver() const
 	
 	TSharedRef<SCustomDialog> AddSolverDialog = SNew(SCustomDialog)
 		.Title(FText(LOCTEXT("EditorController_IKRigFirstSolver", "Add Default Solver")))
-		.DialogContent(SolverOptionBox)
+		.Content()
+		[
+			SolverOptionBox
+		]
 		.Buttons({
 			SCustomDialog::FButton(LOCTEXT("OK", "OK")),
 			SCustomDialog::FButton(LOCTEXT("Cancel", "Cancel"))
diff --git a/Engine/Plugins/Experimental/VirtualProduction/LedWallCalibration/Source/LedWallCalibrationEditor/LedWallCalibrationEditor.Build.cs b/Engine/Plugins/Experimental/VirtualProduction/LedWallCalibration/Source/LedWallCalibrationEditor/LedWallCalibrationEditor.Build.cs
--- a/Engine/Plugins/Experimental/VirtualProduction/LedWallCalibration/Source/LedWallCalibrationEditor/LedWallCalibrationEditor.Build.cs
+++ b/Engine/Plugins/Experimental/VirtualProduction/LedWallCalibration/Source/LedWallCalibrationEditor/LedWallCalibrationEditor.Build.cs
@@ -29,6 +29,7 @@ public class LedWallCalibrationEditor : ModuleRules
 				"Settings",
 				"Slate",
 				"SlateCore",
+				"ToolWidgets",
 				"UnrealEd",
 			}
 		);
diff --git a/Engine/Plugins/Experimental/VirtualProduction/LedWallCalibration/Source/LedWallCalibrationEditor/Private/CalibrationPointArucosForWallDetailsRow.cpp b/Engine/Plugins/Experimental/VirtualProduction/LedWallCalibration/Source/LedWallCalibrationEditor/Private/CalibrationPointArucosForWallDetailsRow.cpp
--- a/Engine/Plugins/Experimental/VirtualProduction/LedWallCalibration/Source/LedWallCalibrationEditor/Private/CalibrationPointArucosForWallDetailsRow.cpp
+++ b/Engine/Plugins/Experimental/VirtualProduction/LedWallCalibration/Source/LedWallCalibrationEditor/Private/CalibrationPointArucosForWallDetailsRow.cpp
@@ -7,7 +7,7 @@
 #include "CalibrationPointComponent.h"
 #include "Components/StaticMeshComponent.h"
 #include "DetailWidgetRow.h"
-#include "Dialogs/CustomDialog.h"
+#include "Dialog/SCustomDialog.h"
 #include "Dialogs/DlgPickAssetPath.h"
 #include "Engine/BlueprintGeneratedClass.h"
 #include "Engine/Texture2D.h"
@@ -96,10 +96,10 @@ void FCalibrationPointArucosForWallDetailsRow::CreateArucos(
 	TSharedRef<SCustomDialog> OptionsWindow =
 		SNew(SCustomDialog)
 		.Title(LOCTEXT("ArucoGenerationOptions", "Aruco Generation Options"))
-		.DialogContent
-		(
-			ArucoGenerationOptionsDetailsView->GetWidget()
-		)
+		.Content()
+		[
+			ArucoGenerationOptionsDetailsView->GetWidget().ToSharedRef()
+		]
 		.Buttons
 		({
 			SCustomDialog::FButton(LOCTEXT("Ok", "Ok")),
diff --git a/Engine/Plugins/VirtualProduction/CameraCalibration/Source/CameraCalibrationEditor/CameraCalibrationEditor.Build.cs b/Engine/Plugins/VirtualProduction/CameraCalibration/Source/CameraCalibrationEditor/CameraCalibrationEditor.Build.cs
--- a/Engine/Plugins/VirtualProduction/CameraCalibration/Source/CameraCalibrationEditor/CameraCalibrationEditor.Build.cs
+++ b/Engine/Plugins/VirtualProduction/CameraCalibration/Source/CameraCalibrationEditor/CameraCalibrationEditor.Build.cs
@@ -41,6 +41,7 @@ namespace UnrealBuildTool.Rules
 					"Slate",
 					"SlateCore",
 					"ToolMenus",
+					"ToolWidgets",
 					"UnrealEd",
 					"WorkspaceMenuStructure",
 				});
diff --git a/Engine/Plugins/VirtualProduction/CameraCalibration/Source/CameraCalibrationEditor/Private/AssetEditor/SImageCenterToolPanel.cpp b/Engine/Plugins/VirtualProduction/CameraCalibration/Source/CameraCalibrationEditor/Private/AssetEditor/SImageCenterToolPanel.cpp
--- a/Engine/Plugins/VirtualProduction/CameraCalibration/Source/CameraCalibrationEditor/Private/AssetEditor/SImageCenterToolPanel.cpp
+++ b/Engine/Plugins/VirtualProduction/CameraCalibration/Source/CameraCalibrationEditor/Private/AssetEditor/SImageCenterToolPanel.cpp
@@ -3,7 +3,7 @@
 #include "SImageCenterToolPanel.h"
 
 #include "CameraCalibrationSubsystem.h"
-#include "Dialogs/CustomDialog.h"
+#include "Dialog/SCustomDialog.h"
 #include "EditorFontGlyphs.h"
 #include "ImageCenterTool.h"
 #include "UI/CameraCalibrationWidgetHelpers.h"
@@ -194,7 +194,10 @@ TSharedRef<SWidget> SImageCenterToolPanel::BuildImageCenterAlgoPickerWidget()
 				TSharedRef<SCustomDialog> AlgoHelpWindow =
 					SNew(SCustomDialog)
 					.Title(FText::FromName(ImageCenterTool->FriendlyName()))
-					.DialogContent(Algo->BuildHelpWidget())
+					.Content()
+					[
+						Algo->BuildHelpWidget()
+					]
 					.Buttons({ SCustomDialog::FButton(LOCTEXT("Ok", "Ok")) });
 
 				AlgoHelpWindow->Show();
diff --git a/Engine/Plugins/VirtualProduction/CameraCalibration/Source/CameraCalibrationEditor/Private/AssetEditor/SLensDistortionToolPanel.cpp b/Engine/Plugins/VirtualProduction/CameraCalibration/Source/CameraCalibrationEditor/Private/AssetEditor/SLensDistortionToolPanel.cpp
--- a/Engine/Plugins/VirtualProduction/CameraCalibration/Source/CameraCalibrationEditor/Private/AssetEditor/SLensDistortionToolPanel.cpp
+++ b/Engine/Plugins/VirtualProduction/CameraCalibration/Source/CameraCalibrationEditor/Private/AssetEditor/SLensDistortionToolPanel.cpp
@@ -5,7 +5,7 @@
 #include "AssetRegistry/AssetData.h"
 #include "CameraCalibrationSubsystem.h"
 #include "CameraLensDistortionAlgo.h"
-#include "Dialogs/CustomDialog.h"
+#include "Dialog/SCustomDialog.h"
 #include "EditorFontGlyphs.h"
 #include "EditorStyleSet.h"
 #include "Engine/Selection.h"
@@ -206,7 +206,10 @@ TSharedRef<SWidget> SLensDistortionToolPanel::BuildAlgoPickerWidget()
 				TSharedRef< SCustomDialog> AlgoHelpWindow = 
 					SNew(SCustomDialog)
 					.Title(FText::FromName(Tool->FriendlyName()))
-					.DialogContent(Algo->BuildHelpWidget())
+					.Content()
+					[
+						Algo->BuildHelpWidget()
+					]
 					.Buttons({
 						SCustomDialog::FButton(LOCTEXT("Ok", "Ok")),
 					});
diff --git a/Engine/Plugins/VirtualProduction/CameraCalibration/Source/CameraCalibrationEditor/Private/AssetEditor/SNodalOffsetToolPanel.cpp b/Engine/Plugins/VirtualProduction/CameraCalibration/Source/CameraCalibrationEditor/Private/AssetEditor/SNodalOffsetToolPanel.cpp
--- a/Engine/Plugins/VirtualProduction/CameraCalibration/Source/CameraCalibrationEditor/Private/AssetEditor/SNodalOffsetToolPanel.cpp
+++ b/Engine/Plugins/VirtualProduction/CameraCalibration/Source/CameraCalibrationEditor/Private/AssetEditor/SNodalOffsetToolPanel.cpp
@@ -5,7 +5,7 @@
 #include "AssetRegistry/AssetData.h"
 #include "CameraCalibrationSubsystem.h"
 #include "CameraNodalOffsetAlgo.h"
-#include "Dialogs/CustomDialog.h"
+#include "Dialog/SCustomDialog.h"
 #include "EditorFontGlyphs.h"
 #include "EditorStyleSet.h"
 #include "Engine/Selection.h"
@@ -205,7 +205,10 @@ TSharedRef<SWidget> SNodalOffsetToolPanel::BuildNodalOffsetAlgoPickerWidget()
 				TSharedRef<SCustomDialog> AlgoHelpWindow =
 					SNew(SCustomDialog)
 					.Title(FText::FromName(NodalOffsetTool->FriendlyName()))
-					.DialogContent(Algo->BuildHelpWidget())
+					.Content()
+					[
+						Algo->BuildHelpWidget()
+					]
 					.Buttons({ SCustomDialog::FButton(LOCTEXT("Ok", "Ok")) });
 
 				AlgoHelpWindow->Show();
diff --git a/Engine/Plugins/VirtualProduction/CameraCalibration/Source/CameraCalibrationEditor/Private/UI/CameraCalibrationWidgetHelpers.cpp b/Engine/Plugins/VirtualProduction/CameraCalibration/Source/CameraCalibrationEditor/Private/UI/CameraCalibrationWidgetHelpers.cpp
--- a/Engine/Plugins/VirtualProduction/CameraCalibration/Source/CameraCalibrationEditor/Private/UI/CameraCalibrationWidgetHelpers.cpp
+++ b/Engine/Plugins/VirtualProduction/CameraCalibration/Source/CameraCalibrationEditor/Private/UI/CameraCalibrationWidgetHelpers.cpp
@@ -3,7 +3,7 @@
 #include "CameraCalibrationWidgetHelpers.h"
 
 #include "AssetEditor/SSimulcamViewport.h"
-#include "Dialogs/CustomDialog.h"
+#include "Dialog/SCustomDialog.h"
 #include "Engine/Texture2D.h"
 #include "Internationalization/Text.h"
 #include "Widgets/Layout/SBox.h"
@@ -70,15 +70,15 @@ void FCameraCalibrationWidgetHelpers::DisplayTextureInWindowAlmostFullScreen(UTe
 		SNew(SCustomDialog)
 		.Title(Title)
 		.ScrollBoxMaxHeight(DetectionWindowMaxHeight)
-		.DialogContent
-		(
+		.Content()
+		[
 			SAssignNew(ViewportWrapper, SBox)
 			.MinDesiredWidth(DetectionWindowMaxWidth)
 			.MinDesiredHeight(DetectionWindowMaxHeight)
 			[
 				SNew(SSimulcamViewport, Texture)
 			]
-		)
+		]
 		.Buttons
 		({
 			SCustomDialog::FButton(LOCTEXT("Ok", "Ok")),
diff --git a/Engine/Plugins/VirtualProduction/LevelSnapshots/Source/LevelSnapshotsEditor/Private/Widgets/SLevelSnapshotsEditorCreationForm.h b/Engine/Plugins/VirtualProduction/LevelSnapshots/Source/LevelSnapshotsEditor/Private/Widgets/SLevelSnapshotsEditorCreationForm.h
--- a/Engine/Plugins/VirtualProduction/LevelSnapshots/Source/LevelSnapshotsEditor/Private/Widgets/SLevelSnapshotsEditorCreationForm.h
+++ b/Engine/Plugins/VirtualProduction/LevelSnapshots/Source/LevelSnapshotsEditor/Private/Widgets/SLevelSnapshotsEditorCreationForm.h
@@ -3,7 +3,7 @@
 #pragma once
 
 #include "CoreMinimal.h"
-#include "Dialogs/CustomDialog.h"
+#include "Dialog/SCustomDialog.h"
 
 DECLARE_DELEGATE_OneParam(FCloseCreationFormDelegate, const FText& /* Description */);
 
diff --git a/Engine/Plugins/VirtualProduction/Switchboard/Source/SwitchboardEditor/Private/SwitchboardSetupWizard.cpp b/Engine/Plugins/VirtualProduction/Switchboard/Source/SwitchboardEditor/Private/SwitchboardSetupWizard.cpp
--- a/Engine/Plugins/VirtualProduction/Switchboard/Source/SwitchboardEditor/Private/SwitchboardSetupWizard.cpp
+++ b/Engine/Plugins/VirtualProduction/Switchboard/Source/SwitchboardEditor/Private/SwitchboardSetupWizard.cpp
@@ -5,7 +5,7 @@
 #include "SwitchboardEditorSettings.h"
 #include "SwitchboardEditorStyle.h"
 #include "SwitchboardScriptInterop.h"
-#include "Dialogs/CustomDialog.h"
+#include "Dialog/SMessageDialog.h"
 #include "Editor.h"
 #include "ISinglePropertyView.h"
 #include "SPrimaryButton.h"
@@ -531,13 +531,9 @@ SSwitchboardSetupWizard::EInstallProgress SSwitchboardSetupWizard::GetProgress()
 
 bool SSwitchboardSetupWizard::ConfirmAndInterruptInstall()
 {
-	TSharedRef<SCustomDialog> ConfirmDialog = SNew(SCustomDialog)
+	TSharedRef<SMessageDialog> ConfirmDialog = SNew(SMessageDialog)
 		.Title(FText(LOCTEXT("Wizard_ConfirmInterruptInstall_WindowTitle", "Cancel Switchboard Installation?")))
-		.DialogContent(
-			SNew(STextBlock)
-			.Margin(FMargin(8.0f, 24.0f, 8.0f, 0.0f))
-			.Text(LOCTEXT("Wizard_ConfirmInterruptInstall_TextBlock", "Are you sure you want to cancel installation?\n\nSwitchboard requires these additional library/framework dependencies to launch."))
-		)
+		.Message(LOCTEXT("Wizard_ConfirmInterruptInstall_TextBlock", "Are you sure you want to cancel installation?\n\nSwitchboard requires these additional library/framework dependencies to launch."))
 		.Buttons({
 			SCustomDialog::FButton(LOCTEXT("Wizard_ConfirmInterruptInstall_CancelInstallButton", "Cancel Installation")),
 			SCustomDialog::FButton(LOCTEXT("Wizard_ConfirmInterruptInstall_ContinueInstallButton", "Continue Installation")),
diff --git a/Engine/Plugins/VirtualProduction/Takes/Source/TakeRecorder/Private/Widgets/STakeRecorderCockpit.cpp b/Engine/Plugins/VirtualProduction/Takes/Source/TakeRecorder/Private/Widgets/STakeRecorderCockpit.cpp
--- a/Engine/Plugins/VirtualProduction/Takes/Source/TakeRecorder/Private/Widgets/STakeRecorderCockpit.cpp
+++ b/Engine/Plugins/VirtualProduction/Takes/Source/TakeRecorder/Private/Widgets/STakeRecorderCockpit.cpp
@@ -56,7 +56,7 @@
 #include "Modules/ModuleManager.h"
 #include "LevelEditor.h"
 #include "ISettingsModule.h"
-#include "Dialogs/CustomDialog.h"
+#include "Dialog/SMessageDialog.h"
 
 extern UNREALED_API UEditorEngine* GEditor;
 
@@ -1006,9 +1006,9 @@ void STakeRecorderCockpit::CancelRecording()
 	UTakeRecorder* CurrentRecording = UTakeRecorder::GetActiveRecorder();
 	if (CurrentRecording)
 	{
-		TSharedRef<SCustomDialog> ConfirmDialog = SNew(SCustomDialog)
+		TSharedRef<SMessageDialog> ConfirmDialog = SNew(SMessageDialog)
 			.Title(FText(LOCTEXT("ConfirmCancelRecordingTitle", "Cancel Recording?")))
-			.DialogContent( SNew(STextBlock).Text(LOCTEXT("ConfirmCancelRecording", "Are you sure you want to cancel the current recording?")))
+			.Message(LOCTEXT("ConfirmCancelRecording", "Are you sure you want to cancel the current recording?"))
 			.Buttons({
 				SCustomDialog::FButton(LOCTEXT("Yes", "Yes"), FSimpleDelegate::CreateLambda([WeakRecording = TWeakObjectPtr<UTakeRecorder>(CurrentRecording)]()
 					{ 
diff --git a/Engine/Plugins/VirtualProduction/Takes/Source/TakeRecorder/TakeRecorder.Build.cs b/Engine/Plugins/VirtualProduction/Takes/Source/TakeRecorder/TakeRecorder.Build.cs
--- a/Engine/Plugins/VirtualProduction/Takes/Source/TakeRecorder/TakeRecorder.Build.cs
+++ b/Engine/Plugins/VirtualProduction/Takes/Source/TakeRecorder/TakeRecorder.Build.cs
@@ -36,6 +36,7 @@ public class TakeRecorder : ModuleRules
 				"WorkspaceMenuStructure",
 				"Analytics",
 				"ToolMenus",
+				"ToolWidgets"
 			}
 		);
 
diff --git a/Engine/Source/Developer/ToolWidgets/Private/Dialog/SCustomDialog.cpp b/Engine/Source/Developer/ToolWidgets/Private/Dialog/SCustomDialog.cpp
new file mode 100644
--- /dev/null
+++ b/Engine/Source/Developer/ToolWidgets/Private/Dialog/SCustomDialog.cpp
@@ -0,0 +1,189 @@
+// Copyright Epic Games, Inc. All Rights Reserved.
+
+#include "Dialog/SCustomDialog.h"
+
+#include "HAL/PlatformApplicationMisc.h"
+
+#include "Framework/Application/SlateApplication.h"
+#include "Framework/Docking/TabManager.h"
+#include "Logging/LogMacros.h"
+#include "Styling/AppStyle.h"
+#include "Styling/SlateBrush.h"
+#include "Widgets/Images/SImage.h"
+#include "Widgets/Input/SButton.h"
+#include "Widgets/Text/STextBlock.h"
+#include "Widgets/Layout/SSpacer.h"
+#include "Widgets/Layout/SBox.h"
+#include "Widgets/Layout/SScrollBox.h"
+#include "Widgets/Layout/SUniformGridPanel.h"
+#include "Widgets/SBoxPanel.h"
+
+DEFINE_LOG_CATEGORY_STATIC(LogCustomDialog, Log, All);
+
+void SCustomDialog::Construct(const FArguments& InArgs)
+{
+	UE_LOG(LogCustomDialog, Log, TEXT("Dialog displayed:"), *InArgs._Title.ToString());
+	check(InArgs._Buttons.Num() > 0);
+	
+	OnClosed = InArgs._OnClosed;
+	SWindow::Construct( SWindow::FArguments()
+		.Title(InArgs._Title)
+		.SizingRule(ESizingRule::Autosized)
+		.SupportsMaximize(false)
+		.SupportsMinimize(false)
+		[
+			SNew(SBorder)
+			.Padding(InArgs._RootPadding)
+			.BorderImage(FAppStyle::Get().GetBrush( "ToolPanel.GroupBorder" ))
+			[
+				SNew(SVerticalBox)
+				+ SVerticalBox::Slot()
+				.FillHeight(1.0f)
+				[
+					CreateContentBox(InArgs)
+				]
+				
+				+ SVerticalBox::Slot()
+				.VAlign(VAlign_Center)
+				.AutoHeight()
+				.Padding(InArgs._ButtonAreaPadding)
+				[
+					CreateButtonBox(InArgs)
+				]
+			]
+		] );
+}
+
+int32 SCustomDialog::ShowModal()
+{
+	FSlateApplication::Get().AddModalWindow(StaticCastSharedRef<SWindow>(this->AsShared()), FGlobalTabmanager::Get()->GetRootWindow());
+	return LastPressedButton;
+}
+
+void SCustomDialog::Show()
+{
+	TSharedRef<SWindow> Window = FSlateApplication::Get().AddWindow(StaticCastSharedRef<SWindow>(this->AsShared()), true);
+	if (OnClosed.IsBound())
+	{
+		Window->GetOnWindowClosedEvent().AddLambda([this](const TSharedRef<SWindow>& Window) { OnClosed.Execute(); });
+	}
+}
+
+TSharedRef<SWidget> SCustomDialog::CreateContentBox(const FArguments& InArgs)
+{
+	TSharedRef<SHorizontalBox> ContentBox = SNew(SHorizontalBox);
+
+	if (InArgs._IconBrush.IsValid())
+	{
+		const FSlateBrush* ImageBrush = FAppStyle::Get().GetBrush(InArgs._IconBrush);
+		if (ensureMsgf(ImageBrush != nullptr, TEXT("Brush %s is unknown"), *InArgs._IconBrush.ToString()))
+		{
+			ContentBox->AddSlot()
+				.AutoWidth()
+				.VAlign(InArgs._VAlignIcon)
+				.HAlign(InArgs._HAlignIcon)
+				.Padding(0, 0, 8, 0)
+				[
+					SNew(SImage)
+					.DesiredSizeOverride(InArgs._IconDesiredSizeOverride)
+					.Image(ImageBrush)
+				];
+		}
+	}
+
+	if (InArgs._UseScrollBox)
+	{
+		ContentBox->AddSlot()
+		.Padding(InArgs._ContentAreaPadding)
+		[
+			SNew(SBox)
+			.MaxDesiredHeight(InArgs._ScrollBoxMaxHeight)
+			[
+				SNew(SScrollBox)
+				+SScrollBox::Slot()
+				[
+					InArgs._Content.Widget
+				]
+			]
+		];
+	}
+	else
+	{
+		ContentBox->AddSlot()
+			.FillWidth(1.0f)
+			.VAlign(VAlign_Center)
+			.HAlign(HAlign_Left)
+			.Padding(InArgs._ContentAreaPadding)
+			[
+				InArgs._Content.Widget
+			];
+	}
+	
+	return ContentBox;
+}
+
+TSharedRef<SWidget> SCustomDialog::CreateButtonBox(const FArguments& InArgs)
+{
+	TSharedPtr<SUniformGridPanel> ButtonPanel;
+	TSharedRef<SHorizontalBox> ButtonBox =
+		SNew(SHorizontalBox)
+	
+		// Before buttons
+		+SHorizontalBox::Slot()
+			.FillWidth(1.f)
+			.HAlign(HAlign_Left)
+			.VAlign(VAlign_Center)
+			[
+				InArgs._BeforeButtons.Widget
+			]
+
+		// Buttons
+		+SHorizontalBox::Slot()
+			.FillWidth(1.0f)
+			.VAlign(VAlign_Center)
+			.HAlign(HAlign_Right)
+			[
+				SAssignNew(ButtonPanel, SUniformGridPanel)
+				.SlotPadding(FAppStyle::Get().GetMargin("StandardDialog.SlotPadding"))
+				.MinDesiredSlotWidth(FAppStyle::Get().GetFloat("StandardDialog.MinDesiredSlotWidth"))
+				.MinDesiredSlotHeight(FAppStyle::Get().GetFloat("StandardDialog.MinDesiredSlotHeight"))
+			];
+
+	for (int32 i = 0; i < InArgs._Buttons.Num(); ++i)
+	{
+		const FButton& Button = InArgs._Buttons[i];
+
+		const FButtonStyle* ButtonStyle = Button.bIsPrimary ?
+				&FAppStyle::Get().GetWidgetStyle< FButtonStyle >( "PrimaryButton" ) :
+				&FAppStyle::Get().GetWidgetStyle< FButtonStyle >("Button");
+		ButtonPanel->AddSlot(ButtonPanel->GetChildren()->Num(), 0)
+		[
+			SNew(SButton)
+			.OnClicked(FOnClicked::CreateSP(this, &SCustomDialog::OnButtonClicked, Button.OnClicked, i))
+			.ButtonStyle(ButtonStyle)
+			[
+				SNew(SHorizontalBox)
+				+ SHorizontalBox::Slot()
+				.VAlign(VAlign_Center)
+				.HAlign(HAlign_Center)
+				[
+					SNew(STextBlock)
+					.Text(Button.ButtonText)
+				]
+			]
+		];
+	}
+
+	return ButtonBox;
+}
+
+/** Handle the button being clicked */
+FReply SCustomDialog::OnButtonClicked(FSimpleDelegate OnClicked, int32 ButtonIndex)
+{
+	LastPressedButton = ButtonIndex;
+
+	FSlateApplication::Get().RequestDestroyWindow(StaticCastSharedRef<SWindow>(this->AsShared()));
+
+	OnClicked.ExecuteIfBound();
+	return FReply::Handled();
+}
diff --git a/Engine/Source/Developer/ToolWidgets/Private/Dialog/SMessageDialog.cpp b/Engine/Source/Developer/ToolWidgets/Private/Dialog/SMessageDialog.cpp
new file mode 100644
--- /dev/null
+++ b/Engine/Source/Developer/ToolWidgets/Private/Dialog/SMessageDialog.cpp
@@ -0,0 +1,67 @@
+// Copyright Epic Games, Inc. All Rights Reserved.
+
+#include "Dialog/SMessageDialog.h"
+
+#include "HAL/PlatformApplicationMisc.h"
+
+#include "Styling/AppStyle.h"
+#include "Styling/SlateBrush.h"
+#include "Widgets/Images/SImage.h"
+#include "Widgets/Input/SButton.h"
+#include "Widgets/Text/STextBlock.h"
+#include "Widgets/Layout/SSpacer.h"
+#include "Widgets/Layout/SScrollBox.h"
+#include "Widgets/SBoxPanel.h"
+#include "Widgets/Input/SHyperlink.h"
+
+void SMessageDialog::Construct(const FArguments& InArgs)
+{
+	Message = InArgs._Message;
+	
+	SCustomDialog::Construct(SCustomDialog::FArguments()
+		.Title(InArgs._Title)
+		.Content()
+		[
+			SNew(STextBlock)
+			.Text(InArgs._Message)
+			.Font(FAppStyle::Get().GetFontStyle("StandardDialog.LargeFont"))
+			.WrapTextAt(InArgs._WrapMessageAt)
+		]
+		.RootPadding(16.f)
+		.Buttons(InArgs._Buttons)
+		.IconBrush(InArgs._IconBrush)
+		.IconDesiredSizeOverride(FVector2D(24.f, 24.f))
+		.HAlignIcon(HAlign_Left)
+		.VAlignIcon(VAlign_Top)
+		.ContentAreaPadding(FMargin(16.f, 0.f, 0.f, 0.f))
+		.UseScrollBox(InArgs._UseScrollBox)
+		.ScrollBoxMaxHeight(InArgs._UseScrollBox)
+		.ButtonAreaPadding(0.f)
+		.OnClosed(InArgs._OnClosed)
+		.BeforeButtons()
+		[
+			SNew(SHyperlink)
+				.OnNavigate(this, &SMessageDialog::CopyMessageToClipboard)
+				// Re-use existing localization of SChoiceDialog
+				.Text( NSLOCTEXT("SChoiceDialog", "CopyMessageHyperlink", "Copy Message") )
+				.ToolTipText( NSLOCTEXT("SChoiceDialog", "CopyMessageTooltip", "Copy the text in this message to the clipboard (CTRL+C)") )
+		]
+		);
+}
+
+FReply SMessageDialog::OnKeyDown(const FGeometry& MyGeometry, const FKeyEvent& InKeyEvent)
+{
+	if (InKeyEvent.GetKey() == EKeys::C && InKeyEvent.IsControlDown())
+	{
+		CopyMessageToClipboard();
+		return FReply::Handled();
+	}
+
+	//if it was some other button, ignore it
+	return FReply::Unhandled();
+}
+
+void SMessageDialog::CopyMessageToClipboard()
+{
+	FPlatformApplicationMisc::ClipboardCopy(*Message.ToString());
+}
diff --git a/Engine/Source/Developer/ToolWidgets/Public/Dialog/SCustomDialog.h b/Engine/Source/Developer/ToolWidgets/Public/Dialog/SCustomDialog.h
new file mode 100644
--- /dev/null
+++ b/Engine/Source/Developer/ToolWidgets/Public/Dialog/SCustomDialog.h
@@ -0,0 +1,140 @@
+// Copyright Epic Games, Inc. All Rights Reserved.
+
+#pragma once
+
+#include "Widgets/SWindow.h"
+
+/**
+ * This is a custom dialog class, which allows any Slate widget to be used as the contents,
+ * with any number of buttons that have any text. 
+ * It also supports adding a custom icon to the dialog.
+ * 
+ * Usage:
+ * TSharedRef<SCustomDialog> HelloWorldDialog = SNew(SCustomDialog)
+		.Title(FText(LOCTEXT("HelloWorldTitleExample", "Hello, World!")))
+		.DialogContent( SNew(SImage).Image(FName(TEXT("Hello"))))
+		.Buttons({
+			SCustomDialog::FButton(LOCTEXT("OK", "OK")),
+			SCustomDialog::FButton(LOCTEXT("Cancel", "Cancel"))
+		});
+
+   // returns 0 when OK is pressed, 1 when Cancel is pressed, -1 if the window is closed
+   const int ButtonPressed = HelloWorldDialog->ShowModal();
+
+ * Note: If the content is only text, see SMessageDialog.  
+ */
+class TOOLWIDGETS_API SCustomDialog : public SWindow
+{
+public:
+	
+	struct FButton
+	{
+		FButton(const FText& InButtonText, const FSimpleDelegate& InOnClicked = FSimpleDelegate())
+			: ButtonText(InButtonText)
+			, OnClicked(InOnClicked)
+		{}
+
+		/** Primary buttons are highlighted */
+		FButton& SetPrimary(bool bValue) { bIsPrimary = bValue; return *this; }
+
+		/** Called when the button is clicked */
+		FButton& SetOnClicked(FSimpleDelegate InOnClicked) { OnClicked = MoveTemp(InOnClicked); return *this; }
+
+		FText ButtonText;
+		FSimpleDelegate OnClicked;
+		bool bIsPrimary = false;
+	};
+
+	SLATE_BEGIN_ARGS(SCustomDialog) 
+		: _HAlignIcon(HAlign_Left)
+		, _VAlignIcon(VAlign_Center)
+		, _RootPadding(FMargin(4.f))
+		, _ButtonAreaPadding(FMargin(20.f, 0.f, 0.f, 0.f))
+		, _UseScrollBox(true)
+		, _ScrollBoxMaxHeight(300)
+	{
+		_AccessibleParams = FAccessibleWidgetData(EAccessibleBehavior::Auto);
+	}
+
+		/*********** Functional ***********/
+		
+		/** Title to display for the dialog. */
+		SLATE_ARGUMENT(FText, Title)
+
+		/** The content to display above the button; icon is optionally created to the left of it.  */
+		SLATE_NAMED_SLOT(FArguments, Content)
+	
+		/** The buttons that this dialog should have. One or more buttons must be added.*/
+		SLATE_ARGUMENT(TArray<FButton>, Buttons)
+
+		/** Event triggered when the dialog is closed, either because one of the buttons is pressed, or the windows is closed. */
+		SLATE_EVENT(FSimpleDelegate, OnClosed)
+	
+
+		/*********** Cosmetic ***********/
+
+		/** Optional icon to display in the dialog (default: none) */
+		SLATE_ARGUMENT(FName, IconBrush)
+	
+		/** When specified, ignore the brushes size and report the DesiredSizeOverride as the desired image size (default: use icon size) */
+		SLATE_ATTRIBUTE(TOptional<FVector2D>, IconDesiredSizeOverride)
+
+		/** Alignment of icon (default: HAlign_Left)*/
+		SLATE_ARGUMENT(EHorizontalAlignment, HAlignIcon)
+
+		/** Alignment of icon (default: VAlign_Center) */
+		SLATE_ARGUMENT(EVerticalAlignment, VAlignIcon)
+	
+	
+		/** Custom widget placed before the buttons */
+		SLATE_NAMED_SLOT(FArguments, BeforeButtons)
+
+
+		/** Padding to apply to the widget embedded in the window, i.e. to all widgets contained in the window (default: {4,4,4,4} )*/
+		SLATE_ATTRIBUTE(FMargin, RootPadding)
+
+		/** Padding to apply around the layout holding the buttons (default: {20,0,0,0}) */
+		SLATE_ATTRIBUTE(FMargin, ButtonAreaPadding)
+
+		/** Padding to apply to DialogContent - you can use it to move away from the icon (default: {0,0,0,0}) */
+		SLATE_ATTRIBUTE(FMargin, ContentAreaPadding)
+	
+
+		/** Should this dialog use a scroll box for over-sized content? (default: true) */
+		SLATE_ARGUMENT(bool, UseScrollBox)
+
+		/** Max height for the scroll box (default: 300) */
+		SLATE_ARGUMENT(int32, ScrollBoxMaxHeight)
+
+
+		/********** Legacy - do not use **********/
+	
+		/** Content for the dialog (deprecated - use Content instead)*/
+		SLATE_ARGUMENT_DEPRECATED(TSharedPtr<SWidget>, DialogContent, 5.1, "Use Content() instead.")
+
+	SLATE_END_ARGS()
+	
+	void Construct(const FArguments& InArgs);
+
+	/** Show the dialog.
+	 * This method will return immediately.
+	 */ 
+	void Show();
+
+	/** Show a modal dialog. Will block until an input is received.
+	 * Returns the index of the button that was pressed.
+	 */
+	int32 ShowModal();
+
+private:
+
+	TSharedRef<SWidget> CreateContentBox(const FArguments& InArgs);
+	TSharedRef<SWidget> CreateButtonBox(const FArguments& InArgs);
+	
+	FReply OnButtonClicked(FSimpleDelegate OnClicked, int32 ButtonIndex);
+
+	/** The index of the button that was pressed last. */
+	int32 LastPressedButton = -1;
+
+	FSimpleDelegate OnClosed;
+};
\ No newline at end of file
diff --git a/Engine/Source/Developer/ToolWidgets/Public/Dialog/SMessageDialog.h b/Engine/Source/Developer/ToolWidgets/Public/Dialog/SMessageDialog.h
new file mode 100644
--- /dev/null
+++ b/Engine/Source/Developer/ToolWidgets/Public/Dialog/SMessageDialog.h
@@ -0,0 +1,61 @@
+// Copyright Epic Games, Inc. All Rights Reserved.
+
+#pragma once
+
+#include "CoreMinimal.h"
+#include "SCustomDialog.h"
+
+/**
+ * Special case of SCustomDialog dedicated to only displaying text messages.
+ * This class enforces uniform style and also adds a hyperlink for copying the message.
+ */
+class TOOLWIDGETS_API SMessageDialog : public SCustomDialog
+{
+public:
+	// Convenience for code using this class
+	using FButton = SCustomDialog::FButton;
+
+	SLATE_BEGIN_ARGS(SMessageDialog)
+		: _WrapMessageAt(512.f)
+	{}
+
+		/********** Functional **********/
+	
+		/** Title to display for the dialog. */
+		SLATE_ARGUMENT(FText, Title)
+	
+		/** Message content */
+		SLATE_ARGUMENT(FText, Message)
+	
+		/** The buttons that this dialog should have. One or more buttons must be added.*/
+		SLATE_ARGUMENT(TArray<FButton>, Buttons)
+
+		/** Event triggered when the dialog is closed, either because one of the buttons is pressed, or the windows is closed. */
+		SLATE_EVENT(FSimpleDelegate, OnClosed)
+	
+		/********** Cosmetic **********/
+	
+		/** Optional icon to display in the dialog. (default: none) */
+		SLATE_ARGUMENT(FName, IconBrush)
+
+		/** Should this dialog use a scroll box for over-sized content? (default: true) */
+		SLATE_ARGUMENT(bool, UseScrollBox)
+
+		/** Max height for the scroll box (default: 300) */
+		SLATE_ARGUMENT(int32, ScrollBoxMaxHeight)
+
+		/** When to wrap the message text (default: 512) */
+		SLATE_ATTRIBUTE(float, WrapMessageAt)
+	
+	SLATE_END_ARGS()
+	
+	void Construct(const FArguments& InArgs);
+
+	virtual	FReply OnKeyDown(const FGeometry& MyGeometry, const FKeyEvent& InKeyEvent) override;
+
+private:
+
+	FText Message;
+	
+	void CopyMessageToClipboard();
+};
\ No newline at end of file
diff --git a/Engine/Source/Developer/ToolWidgets/ToolWidgets.Build.cs b/Engine/Source/Developer/ToolWidgets/ToolWidgets.Build.cs
--- a/Engine/Source/Developer/ToolWidgets/ToolWidgets.Build.cs
+++ b/Engine/Source/Developer/ToolWidgets/ToolWidgets.Build.cs
@@ -12,6 +12,7 @@ public class ToolWidgets : ModuleRules
 		PublicDependencyModuleNames.AddRange(
 			new string[] {
 				"Core",
+				"CoreUObject",
 				"Slate",
 				"SlateCore",
 				"InputCore",
@@ -20,6 +21,7 @@ public class ToolWidgets : ModuleRules
 
 		PrivateDependencyModuleNames.AddRange(
 			new string[] {
+				"ApplicationCore"
 			}
 		);
 
diff --git a/Engine/Source/Editor/Kismet/Private/BlueprintEditor.cpp b/Engine/Source/Editor/Kismet/Private/BlueprintEditor.cpp
--- a/Engine/Source/Editor/Kismet/Private/BlueprintEditor.cpp
+++ b/Engine/Source/Editor/Kismet/Private/BlueprintEditor.cpp
@@ -28,7 +28,7 @@
 #include "Widgets/Views/STableViewBase.h"
 #include "Widgets/Views/STableRow.h"
 #include "Widgets/Views/SListView.h"
-#include "Dialogs/CustomDialog.h"
+#include "Dialog/SCustomDialog.h"
 #include "SCheckBoxList.h"
 #include "EdGraph/EdGraphSchema.h"
 #include "EdGraphNode_Comment.h"
@@ -3985,7 +3985,10 @@ void FBlueprintEditor::DeleteUnusedVariables_OnClicked()
 		TSharedRef<SCustomDialog> CustomDialog = SNew(SCustomDialog)
 			.Title(LOCTEXT("DeleteUnusedVariablesDialog_Title", "Delete Unused Variables"))
 			.IconBrush("NotificationList.DefaultMessage")
-			.DialogContent(DialogContain)
+			.Content()
+			[
+				DialogContain
+			]
 			.Buttons(
 			{
 				SCustomDialog::FButton(LOCTEXT("DeleteUnusedVariablesDialog_ButtonDelete", "Delete")),
diff --git a/Engine/Source/Editor/MainFrame/Private/Menus/LayoutsMenu.cpp b/Engine/Source/Editor/MainFrame/Private/Menus/LayoutsMenu.cpp
--- a/Engine/Source/Editor/MainFrame/Private/Menus/LayoutsMenu.cpp
+++ b/Engine/Source/Editor/MainFrame/Private/Menus/LayoutsMenu.cpp
@@ -25,7 +25,7 @@
 #include "ToolMenus.h"
 // Editor
 #include "Classes/EditorStyleSettings.h"
-#include "Dialogs/CustomDialog.h"
+#include "Dialog/SCustomDialog.h"
 #include "Editor/EditorPerProjectUserSettings.h"
 #include "Frame/MainFrameActions.h"
 #include "LevelViewportActions.h"
@@ -195,7 +195,10 @@ bool FPrivateLayoutsMenu::TrySaveLayoutOrWarnInternal(const FString& InSourceFil
 					const FText CancelText = NSLOCTEXT("Dialogs", "EAppReturnTypeCancel", "Cancel");
 					TSharedRef<SCustomDialog> CustomDialog = SNew(SCustomDialog)
 						.Title(TextTitle)
-						.DialogContent(DialogContents)
+						.Content()
+						[
+							DialogContents
+						]
 						.Buttons({ SCustomDialog::FButton(PreserveValuesText), SCustomDialog::FButton(ClearValuesText), SCustomDialog::FButton(CancelText) });
 
 					// Returns 0 when "Preserve Values" is pressed, 1 when "Clear Values" is pressed, or 2 when Cancel/Esc is pressed
diff --git a/Engine/Source/Editor/UnrealEd/Private/Dialogs/CustomDialog.cpp b/Engine/Source/Editor/UnrealEd/Private/Dialogs/CustomDialog.cpp
deleted file mode 100644
diff --git a/Engine/Source/Editor/UnrealEd/Private/PlayLevel.cpp b/Engine/Source/Editor/UnrealEd/Private/PlayLevel.cpp
--- a/Engine/Source/Editor/UnrealEd/Private/PlayLevel.cpp
+++ b/Engine/Source/Editor/UnrealEd/Private/PlayLevel.cpp
@@ -96,7 +96,7 @@
 #include "Slate/SGameLayerManager.h"
 #include "HAL/PlatformApplicationMisc.h"
 #include "Widgets/Input/SHyperlink.h"
-#include "Dialogs/CustomDialog.h"
+#include "Dialog/SCustomDialog.h"
 
 #include "IHeadMountedDisplay.h"
 #include "IXRTrackingSystem.h"
@@ -1535,7 +1535,10 @@ static bool ShowBlueprintErrorDialog( TArray<UBlueprint*> ErroredBlueprints )
 	CustomDialog = SNew(SCustomDialog)
 		.Title(DialogTitle)
 		.IconBrush("NotificationList.DefaultMessage")
-		.DialogContent(DialogContents)
+		.Content()
+		[
+			DialogContents
+		]
 		.Buttons( { SCustomDialog::FButton(OKText), SCustomDialog::FButton(CancelText) } );
 
 	const int32 ButtonPressed = CustomDialog->ShowModal();
diff --git a/Engine/Source/Editor/UnrealEd/Public/Dialogs/CustomDialog.h b/Engine/Source/Editor/UnrealEd/Public/Dialogs/CustomDialog.h
deleted file mode 100644
diff --git a/Engine/Source/Runtime/SlateCore/Public/Widgets/DeclarativeSyntaxSupport.h b/Engine/Source/Runtime/SlateCore/Public/Widgets/DeclarativeSyntaxSupport.h
--- a/Engine/Source/Runtime/SlateCore/Public/Widgets/DeclarativeSyntaxSupport.h
+++ b/Engine/Source/Runtime/SlateCore/Public/Widgets/DeclarativeSyntaxSupport.h
@@ -198,6 +198,14 @@ template<typename WidgetType> struct TSlateBaseNamedArgs;
 		SLATE_PRIVATE_ARGUMENT_VARIABLE( ArgType, ArgName ); \
 		SLATE_PRIVATE_ARGUMENT_FUNCTION ( ArgType, ArgName )
 
+/**
+ * Use when an argument used to be declared SLATE_ARGUMENT and it should no longer be used
+ */
+#define SLATE_ARGUMENT_DEPRECATED( ArgType, ArgName, DeprecationVersion, DeprecationMessage) \
+	SLATE_PRIVATE_ARGUMENT_VARIABLE( ArgType, ArgName ); \
+	UE_DEPRECATED(DeprecationVersion, DeprecationMessage) \
+	SLATE_PRIVATE_ARGUMENT_FUNCTION ( ArgType, ArgName )
+
 /**
  * Like SLATE_ARGUMENT, but support a default value. e.g.
  * 
