Combo Boxes: Unify behavior across different types to only select on hitting enter. Fix bugs with keyboard navigation and add a visible highlight for the row that currently has keyboard focus

diff --git a/Engine/Source/Developer/ToolWidgets/Private/SSearchableComboBox.cpp b/Engine/Source/Developer/ToolWidgets/Private/SSearchableComboBox.cpp
--- a/Engine/Source/Developer/ToolWidgets/Private/SSearchableComboBox.cpp
+++ b/Engine/Source/Developer/ToolWidgets/Private/SSearchableComboBox.cpp
@@ -50,6 +50,7 @@ void SSearchableComboBox::Construct(const FArguments& InArgs)
 				.ListItemsSource(&FilteredOptionsSource)
 				.OnGenerateRow(this, &SSearchableComboBox::GenerateMenuItemRow)
 				.OnSelectionChanged(this, &SSearchableComboBox::OnSelectionChanged_Internal)
+				.OnKeyDownHandler(this, &SSearchableComboBox::OnKeyDownHandler)
 				.SelectionMode(ESelectionMode::Single)
 				.ExternalScrollbar(InArgs._CustomScrollbar)
 			]
@@ -188,9 +189,10 @@ void SSearchableComboBox::OnMenuOpenChanged(bool bOpen)
 		}
 
 		// Set focus back to ComboBox for users focusing the ListView that just closed
-		TSharedRef<SWidget> ThisRef = AsShared();
-		FSlateApplication::Get().ForEachUser([&ThisRef](FSlateUser& User) {
-			if (User.HasFocusedDescendants(ThisRef))
+		FSlateApplication::Get().ForEachUser([this](FSlateUser& User) 
+		{
+			TSharedRef<SWidget> ThisRef = this->AsShared();
+			if (User.IsWidgetInFocusPath(this->ComboListView))
 			{
 				User.SetFocus(ThisRef);
 			}
@@ -207,6 +209,12 @@ void SSearchableComboBox::OnSelectionChanged_Internal(TSharedPtr<FString> Propos
 		return;
 	}
 
+	// Don't select on navigation, user has to press "Enter" to select the item
+	if (SelectInfo == ESelectInfo::OnNavigation)
+	{
+		return;
+	}
+
 	// Ensure that the proposed selection is different from selected
 	if (ProposedSelection != SelectedItem)
 	{
@@ -214,15 +222,7 @@ void SSearchableComboBox::OnSelectionChanged_Internal(TSharedPtr<FString> Propos
 		OnSelectionChanged.ExecuteIfBound(ProposedSelection, SelectInfo);
 	}
 
-	// close combo as long as the selection wasn't from navigation
-	if (SelectInfo != ESelectInfo::OnNavigation)
-	{
-		this->SetIsOpen(false);
-	}
-	else
-	{
-		ComboListView->RequestScrollIntoView(SelectedItem, 0);
-	}
+	this->SetIsOpen(false);
 }
 
 void SSearchableComboBox::OnSearchTextChanged(const FText& ChangedText)
@@ -261,4 +261,20 @@ FReply SSearchableComboBox::OnButtonClicked()
 	return SComboButton::OnButtonClicked();
 }
 
+FReply SSearchableComboBox::OnKeyDownHandler(const FGeometry& MyGeometry, const FKeyEvent& InKeyEvent)
+{
+	if (InKeyEvent.GetKey() == EKeys::Enter)
+	{
+		// Select the first selected item on hitting enter
+		TArray<TSharedPtr<FString>> SelectedItems = ComboListView->GetSelectedItems();
+		if (SelectedItems.Num() > 0)
+		{
+			OnSelectionChanged_Internal(SelectedItems[0], ESelectInfo::OnKeyPress);
+			return FReply::Handled();
+		}
+	}
+
+	return FReply::Unhandled();
+}
+
 #undef LOCTEXT_NAMESPACE
diff --git a/Engine/Source/Developer/ToolWidgets/Public/SSearchableComboBox.h b/Engine/Source/Developer/ToolWidgets/Public/SSearchableComboBox.h
--- a/Engine/Source/Developer/ToolWidgets/Public/SSearchableComboBox.h
+++ b/Engine/Source/Developer/ToolWidgets/Public/SSearchableComboBox.h
@@ -126,6 +126,8 @@ private:
 	/** Handle clicking on the content menu */
 	virtual FReply OnButtonClicked() override;
 
+	FReply OnKeyDownHandler(const FGeometry& MyGeometry, const FKeyEvent& InKeyEvent);
+
 	/** The item style to use. */
 	const FTableRowStyle* ItemStyle;
 
diff --git a/Engine/Source/Runtime/Slate/Private/Widgets/Input/SComboButton.cpp b/Engine/Source/Runtime/Slate/Private/Widgets/Input/SComboButton.cpp
--- a/Engine/Source/Runtime/Slate/Private/Widgets/Input/SComboButton.cpp
+++ b/Engine/Source/Runtime/Slate/Private/Widgets/Input/SComboButton.cpp
@@ -123,6 +123,12 @@ FReply SComboButton::OnButtonClicked()
 	// Focusing any newly-created widgets must occur after they have been added to the UI root.
 	FReply ButtonClickedReply = FReply::Handled();
 	
+	// Don't try to focus the menu if the menu is closing
+	if (!IsOpen())
+	{
+		return ButtonClickedReply;
+	}
+
 	TSharedPtr<SWidget> WidgetToFocus = WidgetToFocusPtr.Pin();
 	
 	if (bIsFocusable)
diff --git a/Engine/Source/Runtime/Slate/Public/Widgets/Input/SComboBox.h b/Engine/Source/Runtime/Slate/Public/Widgets/Input/SComboBox.h
--- a/Engine/Source/Runtime/Slate/Public/Widgets/Input/SComboBox.h
+++ b/Engine/Source/Runtime/Slate/Public/Widgets/Input/SComboBox.h
@@ -216,6 +216,7 @@ public:
 				.ListItemsSource(InArgs._OptionsSource)
 				.OnGenerateRow(this, &SComboBox< OptionType >::GenerateMenuItemRow)
 				.OnSelectionChanged(this, &SComboBox< OptionType >::OnSelectionChanged_Internal)
+				.OnKeyDownHandler(this, &SComboBox< OptionType >::OnKeyDownHandler)
 				.SelectionMode(ESelectionMode::Single)
 				.ExternalScrollbar(InArgs._CustomScrollbar)
 			];
@@ -524,11 +525,12 @@ private:
 			}
 
 			// Set focus back to ComboBox for users focusing the ListView that just closed
-			TSharedRef<SWidget> ThisRef = AsShared();
-			FSlateApplication::Get().ForEachUser([&ThisRef](FSlateUser& User) {
-				if (User.HasFocusedDescendants(ThisRef))
+			FSlateApplication::Get().ForEachUser([this](FSlateUser& User) 
+			{
+				TSharedRef<SWidget> ThisRef = this->AsShared();
+				if (User.IsWidgetInFocusPath(this->ComboListView))
 				{
-					User.SetFocus(ThisRef, EFocusCause::SetDirectly);
+					User.SetFocus(ThisRef);
 				}
 			});
 
@@ -575,6 +577,23 @@ private:
 		return SComboButton::OnButtonClicked();
 	}
 
+	FReply OnKeyDownHandler(const FGeometry& MyGeometry, const FKeyEvent& InKeyEvent)
+	{
+		if (InKeyEvent.GetKey() == EKeys::Enter)
+		{
+			// Select the first selected item on hitting enter
+			TArray<OptionType> SelectedItems = ComboListView->GetSelectedItems();
+			if (SelectedItems.Num() > 0)
+			{
+				OnSelectionChanged_Internal(SelectedItems[0], ESelectInfo::OnKeyPress);
+				return FReply::Handled();
+			}
+		}
+
+		return FReply::Unhandled();
+	}
+
+
 	/** Play the pressed sound */
 	void PlayPressedSound() const
 	{
diff --git a/Engine/Source/Runtime/SlateCore/Private/Styling/StarshipCoreStyle.cpp b/Engine/Source/Runtime/SlateCore/Private/Styling/StarshipCoreStyle.cpp
--- a/Engine/Source/Runtime/SlateCore/Private/Styling/StarshipCoreStyle.cpp
+++ b/Engine/Source/Runtime/SlateCore/Private/Styling/StarshipCoreStyle.cpp
@@ -1306,7 +1306,7 @@ void FStarshipCoreStyle::SetupComboButtonStyles(TSharedRef<FStyle>& Style)
 		.SetOddRowBackgroundBrush(FSlateNoResource())
 		.SetOddRowBackgroundHoveredBrush(FSlateColorBrush(FStyleColors::Primary))
 
-		.SetSelectorFocusedBrush(FSlateNoResource())
+		.SetSelectorFocusedBrush(FSlateRoundedBoxBrush(FStyleColors::Transparent, InputFocusRadius, FStyleColors::Select, InputFocusThickness))
 
 		.SetActiveBrush(FSlateNoResource())
 		.SetActiveHoveredBrush(FSlateColorBrush(FStyleColors::Primary))
