SSearchableComboBox: Partially revert change 19356303 (commit selection on Enter for keyboard navigation) due to potential fallout.

diff --git a/Engine/Source/Developer/ToolWidgets/Private/SSearchableComboBox.cpp b/Engine/Source/Developer/ToolWidgets/Private/SSearchableComboBox.cpp
--- a/Engine/Source/Developer/ToolWidgets/Private/SSearchableComboBox.cpp
+++ b/Engine/Source/Developer/ToolWidgets/Private/SSearchableComboBox.cpp
@@ -209,12 +209,6 @@ void SSearchableComboBox::OnSelectionChanged_Internal(TSharedPtr<FString> Propos
 		return;
 	}
 
-	// Don't select on navigation, user has to press "Enter" to select the item
-	if (SelectInfo == ESelectInfo::OnNavigation)
-	{
-		return;
-	}
-
 	// Ensure that the proposed selection is different from selected
 	if (ProposedSelection != SelectedItem)
 	{
@@ -222,7 +216,15 @@ void SSearchableComboBox::OnSelectionChanged_Internal(TSharedPtr<FString> Propos
 		OnSelectionChanged.ExecuteIfBound(ProposedSelection, SelectInfo);
 	}
 
-	this->SetIsOpen(false);
+	// close combo as long as the selection wasn't from navigation
+	if (SelectInfo != ESelectInfo::OnNavigation)
+	{
+		this->SetIsOpen(false);
+	}
+	else
+	{
+		ComboListView->RequestScrollIntoView(SelectedItem, 0);
+	}
 }
 
 void SSearchableComboBox::OnSearchTextChanged(const FText& ChangedText)
