Improve usability of SCustomDialog related widgets:
- Add config for HAlign and VAlign
- Specify auto focus behaviour for buttons
- Override SWindow::FArguments

diff --git a/Engine/Source/Developer/ToolWidgets/Private/Dialog/SCustomDialog.cpp b/Engine/Source/Developer/ToolWidgets/Private/Dialog/SCustomDialog.cpp
--- a/Engine/Source/Developer/ToolWidgets/Private/Dialog/SCustomDialog.cpp
+++ b/Engine/Source/Developer/ToolWidgets/Private/Dialog/SCustomDialog.cpp
@@ -26,7 +26,8 @@ void SCustomDialog::Construct(const FArguments& InArgs)
 	check(InArgs._Buttons.Num() > 0);
 	
 	OnClosed = InArgs._OnClosed;
-	SWindow::Construct( SWindow::FArguments()
+	
+	SWindow::Construct( SWindow::FArguments(InArgs._WindowArguments)
 		.Title(InArgs._Title)
 		.SizingRule(ESizingRule::Autosized)
 		.SupportsMaximize(false)
@@ -62,10 +63,10 @@ int32 SCustomDialog::ShowModal()
 
 void SCustomDialog::Show()
 {
-	TSharedRef<SWindow> Window = FSlateApplication::Get().AddWindow(StaticCastSharedRef<SWindow>(this->AsShared()), true);
+	const TSharedRef<SWindow> Window = FSlateApplication::Get().AddWindow(StaticCastSharedRef<SWindow>(this->AsShared()), true);
 	if (OnClosed.IsBound())
 	{
-		Window->GetOnWindowClosedEvent().AddLambda([this](const TSharedRef<SWindow>& Window) { OnClosed.Execute(); });
+		Window->GetOnWindowClosedEvent().AddLambda([this](const TSharedRef<SWindow>&) { OnClosed.Execute(); });
 	}
 }
 
@@ -94,6 +95,8 @@ TSharedRef<SWidget> SCustomDialog::CreateContentBox(const FArguments& InArgs)
 	if (InArgs._UseScrollBox)
 	{
 		ContentBox->AddSlot()
+		.VAlign(InArgs._VAlignContent)
+		.HAlign(InArgs._HAlignContent)
 		.Padding(InArgs._ContentAreaPadding)
 		[
 			SNew(SBox)
@@ -111,8 +114,8 @@ TSharedRef<SWidget> SCustomDialog::CreateContentBox(const FArguments& InArgs)
 	{
 		ContentBox->AddSlot()
 			.FillWidth(1.0f)
-			.VAlign(VAlign_Center)
-			.HAlign(HAlign_Left)
+			.VAlign(InArgs._VAlignContent)
+			.HAlign(InArgs._HAlignContent)
 			.Padding(InArgs._ContentAreaPadding)
 			[
 				InArgs._Content.Widget
@@ -149,17 +152,17 @@ TSharedRef<SWidget> SCustomDialog::CreateButtonBox(const FArguments& InArgs)
 				.MinDesiredSlotHeight(FAppStyle::Get().GetFloat("StandardDialog.MinDesiredSlotHeight"))
 			];
 
-	for (int32 i = 0; i < InArgs._Buttons.Num(); ++i)
+	bool bCanFocusLastPrimary = true;
+	bool bFocusedAnyButtonYet = false;
+	for (int32 ButtonIndex = 0; ButtonIndex < InArgs._Buttons.Num(); ++ButtonIndex)
 	{
-		const FButton& Button = InArgs._Buttons[i];
+		const FButton& Button = InArgs._Buttons[ButtonIndex];
 
 		const FButtonStyle* ButtonStyle = Button.bIsPrimary ?
 				&FAppStyle::Get().GetWidgetStyle< FButtonStyle >( "PrimaryButton" ) :
 				&FAppStyle::Get().GetWidgetStyle< FButtonStyle >("Button");
-		ButtonPanel->AddSlot(ButtonPanel->GetChildren()->Num(), 0)
-		[
-			SNew(SButton)
-			.OnClicked(FOnClicked::CreateSP(this, &SCustomDialog::OnButtonClicked, Button.OnClicked, i))
+		TSharedRef<SButton> ButtonWidget = SNew(SButton)
+			.OnClicked(FOnClicked::CreateSP(this, &SCustomDialog::OnButtonClicked, Button.OnClicked, ButtonIndex))
 			.ButtonStyle(ButtonStyle)
 			[
 				SNew(SHorizontalBox)
@@ -170,10 +173,27 @@ TSharedRef<SWidget> SCustomDialog::CreateButtonBox(const FArguments& InArgs)
 					SNew(STextBlock)
 					.Text(Button.ButtonText)
 				]
-			]
+			];
+
+		ButtonPanel->AddSlot(ButtonPanel->GetChildren()->Num(), 0)
+		[
+			ButtonWidget
 		];
-	}
 
+		if (Button.bShouldFocus)
+		{
+			bCanFocusLastPrimary = false;
+		}
+		
+		const bool bIsLastButton = ButtonIndex == InArgs._Buttons.Num() - 1;
+		if (Button.bShouldFocus
+			|| (bCanFocusLastPrimary && Button.bIsPrimary)
+			|| (bIsLastButton && !bFocusedAnyButtonYet))
+		{
+			bFocusedAnyButtonYet = true;
+			SetWidgetToFocusOnActivate(ButtonWidget);
+		}
+	}
 	return ButtonBox;
 }
 
diff --git a/Engine/Source/Developer/ToolWidgets/Private/Dialog/SMessageDialog.cpp b/Engine/Source/Developer/ToolWidgets/Private/Dialog/SMessageDialog.cpp
--- a/Engine/Source/Developer/ToolWidgets/Private/Dialog/SMessageDialog.cpp
+++ b/Engine/Source/Developer/ToolWidgets/Private/Dialog/SMessageDialog.cpp
@@ -27,9 +27,12 @@ void SMessageDialog::Construct(const FArguments& InArgs)
 			.Font(FAppStyle::Get().GetFontStyle("StandardDialog.LargeFont"))
 			.WrapTextAt(InArgs._WrapMessageAt)
 		]
+		.WindowArguments(InArgs._WindowArguments)
 		.RootPadding(16.f)
 		.Buttons(InArgs._Buttons)
 		.IconBrush(InArgs._IconBrush)
+		.HAlignContent(HAlign_Fill)
+		.VAlignContent(VAlign_Fill)
 		.IconDesiredSizeOverride(FVector2D(24.f, 24.f))
 		.HAlignIcon(HAlign_Left)
 		.VAlignIcon(VAlign_Top)
diff --git a/Engine/Source/Developer/ToolWidgets/Public/Dialog/SCustomDialog.h b/Engine/Source/Developer/ToolWidgets/Public/Dialog/SCustomDialog.h
--- a/Engine/Source/Developer/ToolWidgets/Public/Dialog/SCustomDialog.h
+++ b/Engine/Source/Developer/ToolWidgets/Public/Dialog/SCustomDialog.h
@@ -2,6 +2,7 @@
 
 #pragma once
 
+#include "CoreMinimal.h"
 #include "Widgets/SWindow.h"
 
 /**
@@ -40,18 +41,29 @@ public:
 		/** Called when the button is clicked */
 		FButton& SetOnClicked(FSimpleDelegate InOnClicked) { OnClicked = MoveTemp(InOnClicked); return *this; }
 
+		/**
+		 * Whether to focus this button. Focus rules:
+		 *  1: If a button has SetFocus(), use the last one
+		 *  2: If a button is marked primary, use the last one
+		 *  3: Otherwise, ise the last button
+		 */
+		FButton& SetFocus() { bShouldFocus = true; return *this; }
+		
 		FText ButtonText;
 		FSimpleDelegate OnClicked;
 		bool bIsPrimary = false;
+		bool bShouldFocus = false;
 	};
 
 	SLATE_BEGIN_ARGS(SCustomDialog) 
-		: _HAlignIcon(HAlign_Left)
+		:  _HAlignIcon(HAlign_Left)
 		, _VAlignIcon(VAlign_Center)
 		, _RootPadding(FMargin(4.f))
 		, _ButtonAreaPadding(FMargin(20.f, 0.f, 0.f, 0.f))
 		, _UseScrollBox(true)
 		, _ScrollBoxMaxHeight(300)
+		, _HAlignContent(HAlign_Left)
+		, _VAlignContent(VAlign_Center)
 	{
 		_AccessibleParams = FAccessibleWidgetData(EAccessibleBehavior::Auto);
 	}
@@ -69,6 +81,9 @@ public:
 
 		/** Event triggered when the dialog is closed, either because one of the buttons is pressed, or the windows is closed. */
 		SLATE_EVENT(FSimpleDelegate, OnClosed)
+
+		/** Provides default values for SWindow::FArguments not overriden by SCustomDialog. */
+		SLATE_ARGUMENT(SWindow::FArguments, WindowArguments)
 	
 
 		/*********** Cosmetic ***********/
@@ -106,6 +121,13 @@ public:
 		/** Max height for the scroll box (default: 300) */
 		SLATE_ARGUMENT(int32, ScrollBoxMaxHeight)
 
+	
+		/** HAlign to use for Content slot (default: HAlign_Left) */
+		SLATE_ARGUMENT(EHorizontalAlignment, HAlignContent)
+	
+		/** VAlign to use for Content slot (default: VAlign_Center) */
+		SLATE_ARGUMENT(EVerticalAlignment, VAlignContent)
+
 
 		/********** Legacy - do not use **********/
 	
@@ -132,9 +154,9 @@ private:
 	TSharedRef<SWidget> CreateButtonBox(const FArguments& InArgs);
 	
 	FReply OnButtonClicked(FSimpleDelegate OnClicked, int32 ButtonIndex);
-
+	
 	/** The index of the button that was pressed last. */
 	int32 LastPressedButton = -1;
 
 	FSimpleDelegate OnClosed;
-};
\ No newline at end of file
+};
diff --git a/Engine/Source/Developer/ToolWidgets/Public/Dialog/SMessageDialog.h b/Engine/Source/Developer/ToolWidgets/Public/Dialog/SMessageDialog.h
--- a/Engine/Source/Developer/ToolWidgets/Public/Dialog/SMessageDialog.h
+++ b/Engine/Source/Developer/ToolWidgets/Public/Dialog/SMessageDialog.h
@@ -16,7 +16,7 @@ public:
 	using FButton = SCustomDialog::FButton;
 
 	SLATE_BEGIN_ARGS(SMessageDialog)
-		: _WrapMessageAt(512.f)
+		:  _WrapMessageAt(512.f)
 	{}
 
 		/********** Functional **********/
@@ -32,6 +32,9 @@ public:
 
 		/** Event triggered when the dialog is closed, either because one of the buttons is pressed, or the windows is closed. */
 		SLATE_EVENT(FSimpleDelegate, OnClosed)
+
+		/** Provides default values for SWindow::FArguments not overriden by SCustomDialog. */
+		SLATE_ARGUMENT(SWindow::FArguments, WindowArguments)
 	
 		/********** Cosmetic **********/
 	
