SFilterBar: Add functionality to create and edit custom text filters from the Add Filter dropdown
Add an SFilterSearchBox widget to allow saving custom searches as filters

diff --git a/Engine/Source/Developer/ToolWidgets/Private/Filters/SCustomTextFilterDialog.cpp b/Engine/Source/Developer/ToolWidgets/Private/Filters/SCustomTextFilterDialog.cpp
new file mode 100644
--- /dev/null
+++ b/Engine/Source/Developer/ToolWidgets/Private/Filters/SCustomTextFilterDialog.cpp
@@ -0,0 +1,278 @@
+// Copyright Epic Games, Inc. All Rights Reserved.
+#include "Filters/SCustomTextFilterDialog.h"
+
+#include "Widgets/Input/SEditableTextBox.h"
+#include "Widgets/Colors/SColorBlock.h"
+#include "Widgets/Colors/SColorPicker.h"
+#include "Widgets/Text/STextBlock.h"
+#include "SPrimaryButton.h"
+
+#define LOCTEXT_NAMESPACE "CustomTextFilter"
+
+void SCustomTextFilterDialog::Construct( const FArguments& InArgs )
+{
+	FilterData = InArgs._FilterData;
+	bInEditMode = InArgs._InEditMode;
+	OnCreateFilter = InArgs._OnCreateFilter;
+	OnDeleteFilter = InArgs._OnDeleteFilter;
+	OnCancelClicked = InArgs._OnCancelClicked;
+	OnModifyFilter = InArgs._OnModifyFilter;
+	
+	TSharedPtr<SHorizontalBox> ButtonBox = SNew(SHorizontalBox);
+
+	// Buttons for when we are editing a filter
+	if(bInEditMode)
+	{
+		// Button to delete the current filter
+		ButtonBox->AddSlot()
+		.FillWidth(1.0)
+		.Padding(0, 0, 16, 0)
+		.HAlign(HAlign_Right)
+		[
+			SNew(SButton)
+			.Text(LOCTEXT("DeleteButton", "Delete"))
+			.OnClicked(this, &SCustomTextFilterDialog::OnDeleteButtonClicked)
+		];
+
+		// Button to create a new filter with the current data
+		ButtonBox->AddSlot()
+		.AutoWidth()
+		.Padding(0, 0, 16, 0)
+		.HAlign(HAlign_Right)
+		[
+			SNew(SButton)
+			.Text(LOCTEXT("CreateNewButton", "Create New"))
+			.OnClicked(this, &SCustomTextFilterDialog::OnCreateFilterButtonClicked, false)
+		];
+
+		// Button to modify the current filter
+		ButtonBox->AddSlot()
+		.AutoWidth()
+		.Padding(0, 0, 16, 0)
+		.HAlign(HAlign_Right)
+		[
+			SNew(SPrimaryButton)
+			.Text(LOCTEXT("ModifyFilterButton", "Modify"))
+			.OnClicked(this, &SCustomTextFilterDialog::OnModifyButtonClicked)
+		];
+	}
+	else
+	{
+		// Button to create a new filter and enable it
+		ButtonBox->AddSlot()
+		.FillWidth(1.0)
+		.Padding(0, 0, 16, 0)
+		.HAlign(HAlign_Right)
+		[
+			SNew(SButton)
+			.Text(LOCTEXT("CreateAndApplyButton", "Create and Apply"))
+			.OnClicked(this, &SCustomTextFilterDialog::OnCreateFilterButtonClicked, true)
+		];
+
+		// Button to create a new filter
+		ButtonBox->AddSlot()
+		.AutoWidth()
+		.Padding(0, 0, 16, 0)
+		.HAlign(HAlign_Right)
+		[
+			SNew(SPrimaryButton)
+			.Text(LOCTEXT("CreateButton", "Create"))
+			.OnClicked(this, &SCustomTextFilterDialog::OnCreateFilterButtonClicked, false)
+		];
+		
+	}
+
+	// Button to close the dialog box, common to both modes
+	ButtonBox->AddSlot()
+	.AutoWidth()
+	.Padding(0, 0, 16, 0)
+	.HAlign(HAlign_Right)
+	[
+		SNew(SButton)
+		.Text(LOCTEXT("CancelButton", "Cancel"))
+		.OnClicked(this, &SCustomTextFilterDialog::OnCancelButtonClicked)
+	];
+	
+	ChildSlot
+	[
+		SNew(SBorder)
+		.BorderImage(FAppStyle::GetBrush("Brushes.Panel"))
+		[
+			SNew(SVerticalBox)
+
+			// Filter Name and Color
+			+ SVerticalBox::Slot()
+			.Padding(0, 38, 0, 0)
+			.AutoHeight()
+			[
+				SNew(SHorizontalBox)
+
+				// Filter Name Text
+				+ SHorizontalBox::Slot()
+				.AutoWidth()
+				.VAlign(VAlign_Center)
+				.Padding(34, 0, 0, 0)
+				[
+					SNew(STextBlock)
+					.Text(LOCTEXT("FilterLabelText", "Filter Label"))
+				]
+				// Filter Name Input Box
+				+ SHorizontalBox::Slot()
+				.AutoWidth()
+				.Padding(21, 0, 0, 0)
+				[
+					SAssignNew(FilterLabelTextBox, SEditableTextBox)
+					.MinDesiredWidth(254)
+					.Text_Lambda([this] ()
+					{
+						return this->FilterData.FilterLabel;
+					})
+					.OnTextChanged_Lambda([this] (const FText& InText)
+					{
+						this->FilterData.FilterLabel = InText;
+					})
+				]
+
+				// Filter Color Text
+				+ SHorizontalBox::Slot()
+				.AutoWidth()
+				.VAlign(VAlign_Center)
+				.Padding(21, 0, 0, 0)
+				[
+					SNew(STextBlock)
+					.Text(LOCTEXT("FilterColorText", "Color"))
+				]
+				// Filter Color
+				+ SHorizontalBox::Slot()
+				.AutoWidth()
+				.Padding(21, 0, 0, 0)
+				[
+					SAssignNew(ColorBlock, SColorBlock )
+					.Color_Lambda([this] ()
+					{
+						return this->FilterData.FilterColor;
+					})
+					.CornerRadius(FVector4(4.0f, 4.0f, 4.0f, 4.0f))
+					.Size(FVector2D(70.0f, 22.0f))
+					.AlphaDisplayMode(EColorBlockAlphaDisplayMode::Ignore)
+					.OnMouseButtonDown( this, &SCustomTextFilterDialog::ColorBlock_OnMouseButtonDown )
+				]
+				
+			]
+			// Filter String
+			+ SVerticalBox::Slot()
+			.AutoHeight()
+			.Padding(0, 21, 0, 0)
+			[
+				SNew(SHorizontalBox)
+				// Filter String Text
+				+ SHorizontalBox::Slot()
+				.AutoWidth()
+				.VAlign(VAlign_Center)
+				.Padding(31, 0, 0, 0)
+				[
+					SNew(STextBlock)
+					.Text(LOCTEXT("FilterStringText", "Filter String"))
+				]
+				// Filter Name Input Box
+				+ SHorizontalBox::Slot()
+				.AutoWidth()
+				.Padding(21, 0, 0, 0)
+				[
+					SNew(SEditableTextBox)
+					.MinDesiredWidth(556)
+					.Text_Lambda([this] ()
+					{
+						return this->FilterData.FilterString;
+					})
+					.OnTextChanged_Lambda([this] (const FText& InText)
+					{
+						this->FilterData.FilterString = InText;
+					})
+				]
+			]
+			// Buttons
+			+ SVerticalBox::Slot()
+			.VAlign(VAlign_Bottom)
+			.Padding(0, 0, 0, 16)
+			[
+				ButtonBox.ToSharedRef()
+			]
+		]
+	];
+		
+}
+
+FReply SCustomTextFilterDialog::ColorBlock_OnMouseButtonDown(const FGeometry& MyGeometry, const FPointerEvent& MouseEvent)
+{
+	if (MouseEvent.GetEffectingButton() == EKeys::LeftMouseButton)
+	{
+		TArray<FLinearColor*> LinearColorArray;
+		LinearColorArray.Add(&(FilterData.FilterColor));
+
+		FColorPickerArgs PickerArgs;
+		PickerArgs.bIsModal = true;
+		PickerArgs.ParentWidget = ColorBlock;
+		PickerArgs.LinearColorArray = &LinearColorArray;
+
+		OpenColorPicker(PickerArgs);
+
+		return FReply::Handled();
+	}
+	else
+	{
+		return FReply::Unhandled();
+	}
+		
+}
+
+bool SCustomTextFilterDialog::CheckFilterValidity() const
+{
+	if(FilterData.FilterLabel.IsEmpty())
+	{
+		FilterLabelTextBox->SetError(LOCTEXT("EmptyFilterLabelError", "Filter Label cannot be empty"));
+		
+		return false;
+	}
+
+	return true;
+}
+
+FReply SCustomTextFilterDialog::OnCreateFilterButtonClicked(bool bApplyFilter) const
+{
+	if(!CheckFilterValidity())
+	{
+		return FReply::Handled();
+	}
+	
+	OnCreateFilter.ExecuteIfBound(FilterData, bApplyFilter);
+
+	return FReply::Handled();
+}
+
+FReply SCustomTextFilterDialog::OnDeleteButtonClicked() const
+{
+	OnDeleteFilter.ExecuteIfBound();
+	
+	return FReply::Handled();
+}
+	
+FReply SCustomTextFilterDialog::OnCancelButtonClicked() const
+{
+	OnCancelClicked.ExecuteIfBound();
+	
+	return FReply::Handled();
+}
+
+FReply SCustomTextFilterDialog::OnModifyButtonClicked() const
+{
+	if(!CheckFilterValidity())
+	{
+		return FReply::Handled();
+	}
+	
+	OnModifyFilter.ExecuteIfBound(FilterData);
+
+	return FReply::Handled();
+}
+#undef LOCTEXT_NAMESPACE
diff --git a/Engine/Source/Developer/ToolWidgets/Private/Filters/SFilterSearchBox.cpp b/Engine/Source/Developer/ToolWidgets/Private/Filters/SFilterSearchBox.cpp
new file mode 100644
--- /dev/null
+++ b/Engine/Source/Developer/ToolWidgets/Private/Filters/SFilterSearchBox.cpp
@@ -0,0 +1,255 @@
+// Copyright Epic Games, Inc. All Rights Reserved.
+
+#include "Filters/SFilterSearchBox.h"
+#include "Widgets/Input/SSearchBox.h"
+#include "Widgets/Input/SMenuAnchor.h"
+#include "Widgets/Views/SListView.h"
+#include "Widgets/SOverlay.h"
+#include "Widgets/Input/SButton.h"
+#include "Widgets/Images/SImage.h"
+#include "Widgets/Input/SComboButton.h"
+
+#define LOCTEXT_NAMESPACE "CustomTextFilter"
+
+void SFilterSearchBox::Construct( const FArguments& InArgs )
+{
+	MaxSearchHistory = InArgs._MaxSearchHistory;
+	OnTextChanged = InArgs._OnTextChanged;
+	OnTextCommitted = InArgs._OnTextCommitted;
+
+	// Default text shown when there are no items in the search history
+	EmptySearchHistoryText = MakeShareable(new FText(LOCTEXT("EmptySearchHistoryText", "The Search History is Empty")));
+	SearchHistory.Add(EmptySearchHistoryText);
+	
+	ChildSlot
+	[
+		SAssignNew(SearchHistoryBox, SMenuAnchor)
+		.Placement(EMenuPlacement::MenuPlacement_ComboBoxRight)
+		[
+			SNew(SOverlay)
+			
+			+ SOverlay::Slot()
+			[
+				SAssignNew(SearchBox, SSearchBox)
+				.InitialText(InArgs._InitialText)
+				.HintText(InArgs._HintText)
+				.OnTextChanged(this, &SFilterSearchBox::HandleTextChanged)
+				.OnTextCommitted(this, &SFilterSearchBox::HandleTextCommitted)
+				.SelectAllTextWhenFocused( false )
+			]
+			+ SOverlay::Slot()
+			.HAlign(HAlign_Right)
+			[
+				SNew(SHorizontalBox)
+
+				+SHorizontalBox::Slot()
+				.VAlign(VAlign_Center)
+				[
+					// Button to save the currently occuring search 
+					SNew(SButton)
+					.Visibility_Lambda([this]()
+					{
+						// Only visible if there is a search active currently, and the 
+						return this->GetText().IsEmpty() || !this->OnClickedAddSearchHistoryButton.IsBound() ? EVisibility::Collapsed : EVisibility::Visible;
+					})
+					.ButtonStyle(FAppStyle::Get(), "NoBorder")
+					.OnClicked_Lambda([this]()
+					{
+						this->OnClickedAddSearchHistoryButton.ExecuteIfBound(this->GetText());
+						return FReply::Handled();
+					})
+					[
+						SNew(SImage)
+						.Image(FAppStyle::Get().GetBrush("Icons.PlusCircle"))
+					]
+				]
+
+				+SHorizontalBox::Slot()
+				[
+					// Chevron to open the search history dropdown
+					SNew(SButton)
+					.ButtonStyle(FAppStyle::Get(), "NoBorder")
+					.ClickMethod(EButtonClickMethod::MouseDown)
+					.OnClicked(this, &SFilterSearchBox::OnClickedSearchHistory)
+					.ToolTipText(LOCTEXT("SearchHistoryToolTipText", "Click to show the Search History"))
+					[
+						SNew(SImage)
+						.Image(FAppStyle::Get().GetBrush("Icons.ChevronDown"))
+					]
+				]
+			]
+		]
+		.MenuContent
+		(
+			SNew(SBorder)
+			.BorderImage(FAppStyle::GetBrush("Menu.Background"))
+			.Padding( FMargin(2) )
+			[
+				SAssignNew(SearchHistoryListView, SListView< TSharedPtr<FText> >)
+				.ListItemsSource(&SearchHistory)
+				.SelectionMode( ESelectionMode::Single )							// Ideally the mouse over would not highlight while keyboard controls the UI
+				.OnGenerateRow(this, &SFilterSearchBox::MakeSearchHistoryRowWidget)
+				.OnSelectionChanged( this, &SFilterSearchBox::OnSelectionChanged)
+				.ItemHeight(18)
+				.ScrollbarDragFocusCause(EFocusCause::SetDirectly) // Use SetDirect so that clicking the scrollbar doesn't close the suggestions list
+			]
+		)
+	];
+}
+
+/** Handler for when text in the editable text box changed */
+void SFilterSearchBox::HandleTextChanged(const FText& NewText)
+{
+	OnTextChanged.ExecuteIfBound(NewText);
+}
+
+/** Handler for when text in the editable text box changed */
+void SFilterSearchBox::HandleTextCommitted(const FText& NewText, ETextCommit::Type CommitType)
+{
+	// Set the text and execute the delegate
+	SetText(NewText);
+	OnTextCommitted.ExecuteIfBound(NewText, CommitType);
+	
+	UpdateSearchHistory(NewText);
+}
+
+void SFilterSearchBox::SetText(const TAttribute< FText >& InNewText)
+{
+	SearchBox->SetText(InNewText);
+}
+
+FText SFilterSearchBox::GetText() const
+{
+	return SearchBox->GetText();
+}
+
+void SFilterSearchBox::SetError( const FText& InError )
+{
+	SearchBox->SetError(InError);
+}
+
+void SFilterSearchBox::SetError( const FString& InError )
+{
+	SearchBox->SetError(InError);
+}
+
+void SFilterSearchBox::OnFocusLost(const FFocusEvent& InFocusEvent)
+{
+	// Update the search history on focus loss (for cases with SearchBoxes that don't commit on Enter)
+	UpdateSearchHistory(SearchBox->GetText());
+}
+
+/** Called by SListView when the selection changes in the search history list */
+void SFilterSearchBox::OnSelectionChanged( TSharedPtr<FText> NewValue, ESelectInfo::Type SelectInfo )
+{
+	/* Make sure the user can only select an item in the history using the mouse, and that they cannot select the
+	 * Placeholder text for empty search history
+	 */
+	if(SelectInfo != ESelectInfo::OnNavigation && NewValue && NewValue != EmptySearchHistoryText)
+	{
+		SearchBox->SetText(*NewValue);
+		SearchHistoryBox->SetIsOpen(false);
+	}
+}
+
+TSharedRef<ITableRow> SFilterSearchBox::MakeSearchHistoryRowWidget(TSharedPtr<FText> SearchText, const TSharedRef<STableViewBase>& OwnerTable)
+{
+	bool bIsEmptySearchHistory = (SearchText == EmptySearchHistoryText);
+	EHorizontalAlignment TextAlignment = (bIsEmptySearchHistory) ? HAlign_Center : HAlign_Left;
+
+	TSharedPtr<SHorizontalBox> RowWidget = SNew(SHorizontalBox);
+
+	// The actual search text
+	RowWidget->AddSlot()
+	.HAlign(TextAlignment)
+	.VAlign(VAlign_Center)
+	.FillWidth(1.0)
+	[
+		SNew(STextBlock)
+		.Text(*SearchText.Get())
+	];
+
+	// If the Search History is not empty, add the button to save the search history
+	if(!bIsEmptySearchHistory)
+	{
+		RowWidget->AddSlot()
+		.HAlign(HAlign_Right)
+		[
+			SNew(SButton)
+			.ButtonStyle(FAppStyle::Get(), "NoBorder")
+			.OnClicked_Lambda([this, SearchText]()
+			{
+				this->OnClickedAddSearchHistoryButton.ExecuteIfBound(*SearchText.Get());
+				return FReply::Handled();
+			})
+			.Visibility_Lambda([this]()
+			{
+				return this->OnClickedAddSearchHistoryButton.IsBound() ? EVisibility::Visible : EVisibility::Collapsed;
+			})
+			.ToolTipText(LOCTEXT("SearchHistoryToolTipText", "Save this search as a custom filter"))
+			[
+				SNew(SImage)
+				.Image(FAppStyle::Get().GetBrush("Icons.PlusCircle"))
+			]
+		];
+	}
+
+	return SNew(STableRow< TSharedPtr<FString> >, OwnerTable)
+			.ShowSelection(!bIsEmptySearchHistory)
+			[
+				RowWidget.ToSharedRef()
+			];
+}
+
+void SFilterSearchBox::UpdateSearchHistory(const FText &NewText)
+{
+	// Don't save empty searches
+	if(NewText.IsEmpty())
+	{
+		return;
+	}
+
+	// If there is only one item, and it is the placeholder empty search text, remove it
+	if(SearchHistory.Num() == 1 && SearchHistory[0] == EmptySearchHistoryText)
+	{
+		SearchHistory.Empty();
+	}
+
+	// Remove any existing occurances of the current search text, we will re-add it to the top if so
+	SearchHistory.RemoveAll([&NewText](TSharedPtr<FText> SearchHistoryText)
+	{
+		if(SearchHistoryText->CompareTo(NewText) == 0)
+		{
+			return true;
+		}
+		return false;
+	});
+
+	// Insert the current search as the most recent in the history
+	SearchHistory.Insert(MakeShareable(new FText(NewText)), 0);
+
+	// Prune old entries until we are at the Max Search History limit
+	while( SearchHistory.Num() > MaxSearchHistory)
+	{
+		SearchHistory.RemoveAt( SearchHistory.Num()-1 );
+	}
+
+	SearchHistoryListView->RequestListRefresh();
+}
+
+FReply SFilterSearchBox::OnClickedSearchHistory()
+{
+	if(SearchHistoryBox->ShouldOpenDueToClick() && !SearchHistory.IsEmpty())
+	{
+		SearchHistoryBox->SetIsOpen(true);
+		SearchHistoryListView->ClearSelection();
+	}
+	else
+	{
+		SearchHistoryBox->SetIsOpen(false);
+	}
+	
+	return FReply::Handled();
+}
+
+#undef LOCTEXT_NAMESPACE
diff --git a/Engine/Source/Developer/ToolWidgets/Public/Filters/SBasicFilterBar.h b/Engine/Source/Developer/ToolWidgets/Public/Filters/SBasicFilterBar.h
--- a/Engine/Source/Developer/ToolWidgets/Public/Filters/SBasicFilterBar.h
+++ b/Engine/Source/Developer/ToolWidgets/Public/Filters/SBasicFilterBar.h
@@ -16,6 +16,12 @@
 #include "Widgets/SBoxPanel.h"
 #include "Misc/FilterCollection.h"
 #include "UObject/Object.h"
+#include "Core/Public/Misc/TextFilter.h"
+#include "Widgets/SWindow.h"
+#include "Filters/SCustomTextFilterDialog.h"
+#include "SSimpleButton.h"
+#include "Styling/StyleColors.h"
+#include "Filters/SFilterSearchBox.h"
 
 #include "SBasicFilterBar.generated.h"
 
@@ -60,6 +66,9 @@ public:
 	/** Delegate for when filters have changed */
 	DECLARE_DELEGATE_OneParam( FOnExtendAddFilterMenu, UToolMenu* );
 	
+	/** Delegate to create a TTextFilter used to compare FilterType with text queries */
+    DECLARE_DELEGATE_RetVal( TSharedPtr<TTextFilter<FilterType>>, FCreateTextFilter);
+	
  	SLATE_BEGIN_ARGS( SBasicFilterBar<FilterType> ){}
 
  		/** Delegate for when filters have changed */
@@ -70,6 +79,17 @@ public:
 
 		/** Initial List of Custom Filters that will be added to the AddFilter Menu */
 		SLATE_ARGUMENT( TArray<TSharedRef<FFilterBase<FilterType>>>, CustomFilters)
+
+		/** A delegate to create a TTextFilter for FilterType items. If provided, will allow creation of custom text filters
+		 *  from the filter dropdown menu.
+		 */
+		SLATE_ARGUMENT(FCreateTextFilter, CreateTextFilter)
+		
+		/** An SFilterSearchBox that can be attached to this filter bar. When provided along with a CreateTextFilter
+		 *  delegate, allows the user to save searches from the Search Box as text filters for the filter bar.
+		 *	NOTE: Will bind a delegate to SFilterSearchBox::OnClickedAddSearchHistoryButton
+		 */
+        SLATE_ARGUMENT(TSharedPtr<SFilterSearchBox>, FilterSearchBox)
 	
  	SLATE_END_ARGS()
 
@@ -77,6 +97,9 @@ public:
  	{
 		OnFilterChanged = InArgs._OnFilterChanged;
  		OnExtendAddFilterMenu = InArgs._OnExtendAddFilterMenu;
+ 		CreateTextFilter = InArgs._CreateTextFilter;
+
+ 		AttachFilterSearchBox(InArgs._FilterSearchBox);
 
  		// A subclass could be using an external filter collection (SFilterList in ContentBrowser)
  		if(!ActiveFilters)
@@ -128,6 +151,18 @@ public:
 				];
  	}
 
+	/** Attach an SFilterSearchBox to this filter bar, overriding theSFilterSearchBox:: OnClickedAddSearchHistoryButton
+	 *  event to save a search as a filter
+	 */
+	void AttachFilterSearchBox(TSharedPtr<SFilterSearchBox> InFilterSearchBox)
+ 	{
+		if(InFilterSearchBox)
+		{
+			InFilterSearchBox->SetOnClickedAddSearchHistoryButton(
+				 SFilterSearchBox::FOnClickedAddSearchHistoryButton::CreateSP(this, &SBasicFilterBar<FilterType>::CreateCustomTextFilterFromSearch));
+		}
+ 	}
+
 private:
 
 	/** A class for check boxes in the filter list. If you double click a filter checkbox, you will enable it and disable all others */
@@ -1042,6 +1077,232 @@ protected:
  			}
  		}
  	}
+ 	
+ 	/** Handler for when a checkbox next to a custom text filter is clicked */
+    void CustomTextFilterClicked(ECheckBoxState CheckBoxState, TSharedRef<FCustomTextFilter<FilterType>> FrontendFilter)
+    {
+    	if (CheckBoxState == ECheckBoxState::Unchecked)
+    	{
+    		RemoveFilter(FrontendFilter);
+    	}
+    	else
+    	{
+    		TSharedRef<SFilter> NewFilter = AddFilterToBar(FrontendFilter);
+    		NewFilter->SetEnabled(true);
+    	}
+    }
+
+	/** Handler for when a custom text filter is created */
+	void OnCreateCustomTextFilter(const FCustomTextFilterData& InFilterData, bool bApplyFilter)
+	{
+		// Create a TTextFilter for the current filter we are making using the provided delegate
+		TSharedRef<TTextFilter<FilterType>> NewTextFilter = CreateTextFilter.Execute().ToSharedRef();
+		
+		// Create the actual text filter
+		TSharedRef<FCustomTextFilter<FilterType>> NewFilter = MakeShared<FCustomTextFilter<FilterType>>(NewTextFilter);
+
+		// Fill in the data the widget gives us
+		NewFilter->SetDisplayName(InFilterData.FilterLabel);
+		NewFilter->SetColor(InFilterData.FilterColor);
+		NewFilter->SetFilterString(InFilterData.FilterString);
+
+		CustomTextFilters.Add(NewFilter);
+		TSharedRef<SFilter> AddedFilter = AddFilterToBar(NewFilter);
+		AddedFilter->SetEnabled(bApplyFilter);
+
+		// Close the create custom filter dialog 		
+		OnCancelCustomTextFilterDialog();
+	}
+
+	/** Handler for when a custom text filter is modified */
+	void OnModifyCustomTextFilter(const FCustomTextFilterData& InFilterData, TSharedPtr<FCustomTextFilter<FilterType>> InFilter)
+	{
+		// Update the filter with the data the widget has given us
+		if(InFilter)
+		{
+			InFilter->SetDisplayName(InFilterData.FilterLabel);
+			InFilter->SetColor(InFilterData.FilterColor);
+			InFilter->SetFilterString(InFilterData.FilterString);
+		}
+
+		// If the filter we modified is active in the bar, remove and re-add it to update the data
+		for (const TSharedRef<SFilter>& Filter : Filters)
+		{
+			if (Filter->GetFrontendFilter() == InFilter)
+			{
+				bool bWasEnabled = Filter->IsEnabled();
+				
+				RemoveFilter(InFilter.ToSharedRef());
+
+				TSharedRef<SFilter> AddedFilter = AddFilterToBar(InFilter.ToSharedRef());
+				AddedFilter->SetEnabled(bWasEnabled);
+			}
+		}
+
+		OnCancelCustomTextFilterDialog();
+	}
+
+	/** Handler for when a custom text filter is deleted */
+	void OnDeleteCustomTextFilter(const TSharedPtr<FCustomTextFilter<FilterType>> InFilter)
+	{
+		if(InFilter)
+		{
+			CustomTextFilters.RemoveSingle(InFilter.ToSharedRef());
+			RemoveFilter(InFilter.ToSharedRef());
+		}
+		
+		OnCancelCustomTextFilterDialog();
+	}
+
+	/** Handler to close the custom text filter dialog */
+	void OnCancelCustomTextFilterDialog()
+	{
+		if(CustomTextFilterWindow.IsValid())
+		{
+			CustomTextFilterWindow.Pin()->RequestDestroyWindow();
+		}
+	}
+
+	/** Creates a dialog box with the SCustomTextFilterDialog Widget */
+	void CreateCustomTextFilterWindow(const FCustomTextFilterData& CustomTextFilterData, TSharedPtr<FCustomTextFilter<FilterType>> InFilter)
+	{
+		/** If we already have a window, delete it */
+		OnCancelCustomTextFilterDialog();
+
+		// If we have a filter to edit, we are in edit mode
+		bool bInEditMode = InFilter.IsValid();
+		
+		FText WindowTitle = bInEditMode ? LOCTEXT("ModifyCustomTextFilterWindow", "Modify Custom Filter") : LOCTEXT("CreateCustomTextFilterWindow", "Create Custom Filter");
+		
+		TSharedPtr<SWindow> NewTextFilterWindow = SNew(SWindow)
+			.Title(WindowTitle)
+			.HasCloseButton(true)
+			.SupportsMaximize(false)
+			.SupportsMinimize(false)
+			.ClientSize(FVector2D(724, 183));
+		
+		TSharedPtr<SCustomTextFilterDialog> CustomTextFilterDialog =
+			SNew(SCustomTextFilterDialog)
+			.FilterData(CustomTextFilterData)
+			.InEditMode(bInEditMode)
+			.OnCreateFilter(this, &SBasicFilterBar<FilterType>::OnCreateCustomTextFilter)
+			.OnDeleteFilter(this, &SBasicFilterBar<FilterType>::OnDeleteCustomTextFilter, InFilter)
+			.OnModifyFilter(this, &SBasicFilterBar<FilterType>::OnModifyCustomTextFilter, InFilter)
+			.OnCancelClicked(this, &SBasicFilterBar<FilterType>::OnCancelCustomTextFilterDialog);
+
+		NewTextFilterWindow->SetContent(CustomTextFilterDialog.ToSharedRef());
+		FSlateApplication::Get().AddWindow(NewTextFilterWindow.ToSharedRef());
+		
+		CustomTextFilterWindow = NewTextFilterWindow;
+	}
+
+	/** Creates a dialog box to Add a custom text filter */
+	void CreateAddCustomTextFilterWindow()
+	{
+		CreateCustomTextFilterWindow(FCustomTextFilterData(), nullptr);
+	}
+
+	/** Creates a dialog box to Edit an existing custom text filter */
+	void CreateEditCustomTextFilterWindow(TSharedPtr<FCustomTextFilter<FilterType>> InFilter)
+	{
+		if(!InFilter)
+		{
+			return;
+		}
+
+		FCustomTextFilterData CustomTextFilterData;
+		CustomTextFilterData.FilterLabel = InFilter->GetDisplayName();
+		CustomTextFilterData.FilterColor = InFilter->GetColor();
+		CustomTextFilterData.FilterString = InFilter->GetFilterString();
+
+		CreateCustomTextFilterWindow(CustomTextFilterData, InFilter);
+	}
+
+	/** Creates a dialog to add a custom text filter from the given search text */
+	void CreateCustomTextFilterFromSearch(const FText& InSearchText)
+	{
+		FCustomTextFilterData CustomTextFilterData;
+		CustomTextFilterData.FilterLabel = InSearchText;
+		CustomTextFilterData.FilterString = InSearchText;
+
+		CreateCustomTextFilterWindow(CustomTextFilterData, nullptr);
+	}
+	
+	/** Populate the Custom Filters submenu in the Add Filter dropdown */
+	void CreateTextFiltersMenu(UToolMenu* InMenu)
+	{
+		FToolMenuSection& Section = InMenu->AddSection("FilterBarTextFiltersSubmenu");
+		
+		Section.AddMenuEntry(
+			"CreateNewTextFilter",
+			LOCTEXT("CreateNewTextFilter", "Create New Filter"),
+			LOCTEXT("CreateNewTextFilterTooltip", "Create a new text filter"),
+			FSlateIcon(FAppStyle::Get().GetStyleSetName(), "Icons.PlusCircle"),
+			FUIAction(
+				FExecuteAction::CreateSP(this, &SBasicFilterBar<FilterType>::CreateAddCustomTextFilterWindow)
+				)
+		);
+
+		FToolMenuSection& CustomFiltersSection = InMenu->AddSection("FilterBarCustomTextFiltersSection");
+		
+		for (const TSharedRef<FCustomTextFilter<FilterType>>& FrontendFilter : CustomTextFilters)
+		{
+			bool bIsFilterChecked = IsFrontendFilterInUse(FrontendFilter);
+			
+			TSharedPtr<SHorizontalBox> CustomTextFilterWidget = SNew(SHorizontalBox);
+			
+			// Checkbox for the filter
+			CustomTextFilterWidget->AddSlot()
+			.AutoWidth()
+			.Padding(16, 0, 0, 0)
+			[
+				SNew(SCheckBox)
+				.Style(FAppStyle::Get(), "Menu.CheckBox")
+				.OnCheckStateChanged(FOnCheckStateChanged::CreateSP(this, &SBasicFilterBar<FilterType>::CustomTextFilterClicked, FrontendFilter))
+				.IsChecked(bIsFilterChecked)
+			];
+
+			// A button showing the filter name, that checks it when clicked on
+			CustomTextFilterWidget->AddSlot()
+			.FillWidth(1.0)
+			[
+				SNew(SButton)
+				.ButtonStyle(FAppStyle::Get(), "SimpleButton")
+				.ForegroundColor(FStyleColors::White)
+				.Text(FrontendFilter->GetDisplayName())
+				.OnClicked_Lambda([this, &FrontendFilter]()
+				{
+					this->FrontendFilterClicked(FrontendFilter);
+					FSlateApplication::Get().DismissAllMenus();
+					return FReply::Handled();
+				})
+			];
+
+			// An edit button to open the edit custom text filter dialog
+			CustomTextFilterWidget->AddSlot()
+			.AutoWidth()
+			.HAlign(HAlign_Right)
+			.Padding(0, 0, 16, 0)
+			[
+				SNew(SSimpleButton)
+				.Icon(FAppStyle::GetBrush("Icons.Edit"))
+				.OnClicked_Lambda([this, &FrontendFilter]()
+				{
+					this->CreateEditCustomTextFilterWindow(FrontendFilter);
+					return FReply::Handled();
+				})
+			];
+
+			CustomFiltersSection.AddEntry(
+			FToolMenuEntry::InitWidget(
+				NAME_None,
+				CustomTextFilterWidget.ToSharedRef(),
+				FText::GetEmpty(),
+				true,
+				true
+			));
+		}
+	}
 
 	/** Helper function to add common sections to the Add Filter Menu */
 	void PopulateCommonFilterSections(UToolMenu* Menu)
@@ -1056,6 +1317,19 @@ protected:
 				FExecuteAction::CreateSP(this, &SBasicFilterBar<FilterType>::OnResetFilters),
 				FCanExecuteAction::CreateLambda([this]() { return HasAnyFilters(); }))
 		);
+
+		// Only add the custom text filter submenu if we have a valid CreateTextFilter delegate to use
+		if(CreateTextFilter.IsBound())
+		{
+			FToolMenuSection& CustomFiltersSection = Menu->AddSection("FilterBarTextFilters");
+
+			CustomFiltersSection.AddSubMenu(
+					"CustomFiltersSubMenu",
+					LOCTEXT("FilterBarTextFilters", "Custom Filters"),
+					LOCTEXT("FilterBarTextFiltersTooltip", "Custom Filters"),
+					FNewToolMenuDelegate::CreateSP(this, &SBasicFilterBar<FilterType>::CreateTextFiltersMenu)
+					);
+		}
 		
 		OnExtendAddFilterMenu.ExecuteIfBound(Menu);
 	}
@@ -1143,6 +1417,15 @@ protected:
 	/** Delegate to extend the AddFilter Menu */
 	FOnExtendAddFilterMenu OnExtendAddFilterMenu;
 
+	/** A delegate used to create a TTextFilter for FilterType */
+	FCreateTextFilter CreateTextFilter;
+	
+	/** Custom text filters that a user can create */
+	TArray< TSharedRef< FCustomTextFilter<FilterType> > > CustomTextFilters;
+
+	/** The window containing the custom text filter dialog */
+	TWeakPtr<SWindow> CustomTextFilterWindow;
+	
 	friend struct FFrontendFilterExternalActivationHelper<FilterType>;
 };
 
diff --git a/Engine/Source/Developer/ToolWidgets/Public/Filters/SCustomTextFilterDialog.h b/Engine/Source/Developer/ToolWidgets/Public/Filters/SCustomTextFilterDialog.h
new file mode 100644
--- /dev/null
+++ b/Engine/Source/Developer/ToolWidgets/Public/Filters/SCustomTextFilterDialog.h
@@ -0,0 +1,211 @@
+// Copyright Epic Games, Inc. All Rights Reserved.
+
+#include "Widgets/SCompoundWidget.h"
+#include "Filters/FilterBase.h"
+#include "Core/Public/Misc/TextFilter.h"
+
+class SColorBlock;
+class SEditableTextBox;
+
+/* A filter that allows testing items of type FilterType against text expressions */
+template<typename FilterType>
+class FCustomTextFilter : public FFilterBase<FilterType>
+{
+public:
+
+	FCustomTextFilter(TSharedPtr<TTextFilter<FilterType>> InTextFilter)
+	: FFilterBase<FilterType>(nullptr)
+	, TextFilter(InTextFilter)
+	{
+		
+	}
+
+	/** Returns the system name for this filter */
+	virtual FString GetName() const override
+	{
+		return FString("FCustomTextFilter");
+	}
+
+	/** Returns the human readable name for this filter */
+	virtual FText GetDisplayName() const override
+	{
+		return DisplayName; 
+	}
+
+	/** Set the human readable name for this filter */
+	void SetDisplayName(const FText& InDisplayName)
+	{
+		DisplayName = InDisplayName;
+	}
+
+	/** Returns the tooltip for this filter, shown in the filters menu */
+	virtual FText GetToolTipText() const override
+	{
+		return TextFilter->GetRawFilterText(); // The tooltip will display the filter string
+	}
+
+	/** Returns the color this filter button will be when displayed as a button */
+	virtual FLinearColor GetColor() const override
+	{
+		return Color;
+	}
+
+	/** Set the color this filter button will be when displayed as a button */
+	void SetColor(const FLinearColor& InColor)
+	{
+		Color = InColor;
+	}
+
+	/** Returns the name of the icon to use in menu entries */
+	virtual FName GetIconName() const override
+	{
+		return NAME_None;
+	}
+
+	/** Returns true if the filter should be in the list when disabled and not in the list when enabled */
+	virtual bool IsInverseFilter() const override
+	{
+		return false;
+	}
+
+	virtual bool PassesFilter( FilterType InItem ) const override
+	{
+		return TextFilter->PassesFilter(InItem);
+	}
+
+	/** Get the actual text this filter is using to test against */
+	FText GetFilterString()
+	{
+		return TextFilter->GetRawFilterText();
+	}
+
+	/** Set the actual text this filter is using to test against */
+	void SetFilterString(const FText& InFilterString)
+	{
+		TextFilter->SetRawFilterText(InFilterString);
+	}
+	
+	// Functionality that is not needed by this filter
+	
+	/** Notification that the filter became active or inactive */
+	virtual void ActiveStateChanged(bool bActive) override
+	{
+		
+	}
+
+	/** Called when the right-click context menu is being built for this filter */
+	virtual void ModifyContextMenu(FMenuBuilder& MenuBuilder) override
+	{
+		
+	}
+	
+	/** Can be overriden for custom FilterBar subclasses to save settings, currently not implemented in any gneeric Filter Bar */
+	virtual void SaveSettings(const FString& IniFilename, const FString& IniSection, const FString& SettingsString) const override
+	{
+		
+	}
+
+	/** Can be overriden for custom FilterBar subclasses to load settings, currently not implemented in any gneeric Filter Bar */
+	virtual void LoadSettings(const FString& IniFilename, const FString& IniSection, const FString& SettingsString) override
+	{
+		
+	}
+	
+protected:
+
+	/* The actual Text Filter containing information about the text being tested against */
+	TSharedPtr<TTextFilter<FilterType>> TextFilter;
+
+	/* The Display Name of this custom filter that the user sees */
+	FText DisplayName;
+
+	/* The Color of this filter pill */
+	FLinearColor Color;
+};
+
+
+/* Struct containing the data that SCustomTextFilterDialog is currently editing */
+struct FCustomTextFilterData
+{
+	FText FilterLabel;
+	FText FilterString;
+	FLinearColor FilterColor;
+
+	FCustomTextFilterData()
+		: FilterColor(FLinearColor::White)
+	{
+		
+	}
+};
+
+class TOOLWIDGETS_API SCustomTextFilterDialog : public SCompoundWidget
+{
+public:
+
+	DECLARE_DELEGATE_TwoParams(FOnCreateFilter, const FCustomTextFilterData& /* InFilterData */, bool /* bApplyFilter */);
+	DECLARE_DELEGATE(FOnDeleteFilter);
+	DECLARE_DELEGATE_OneParam(FOnModifyFilter, const FCustomTextFilterData& /* InFilterData */);
+	DECLARE_DELEGATE(FOnCancelClicked);
+
+	SLATE_BEGIN_ARGS(SCustomTextFilterDialog) {}
+	
+    	/** The filter that this dialog is creating/editing */
+    	SLATE_ARGUMENT(FCustomTextFilterData, FilterData)
+
+		/** True if we are editing an existing filter, false if we are creating a new one */
+		SLATE_ARGUMENT(bool, InEditMode)
+		
+		/** Delegate for when the Create button is clicked */
+		SLATE_EVENT(FOnCreateFilter, OnCreateFilter)
+		
+		/** Delegate for when the Delete button is clicked */
+		SLATE_EVENT(FOnDeleteFilter, OnDeleteFilter)
+		
+		/** Delegate for when the Cancel button is clicked */
+		SLATE_EVENT(FOnCancelClicked, OnCancelClicked)
+
+		/** Delegate for when the Modify Filter button is clicked */
+        SLATE_EVENT(FOnModifyFilter, OnModifyFilter)
+    
+    SLATE_END_ARGS()
+    	
+    /** Constructs this widget with InArgs */
+    void Construct( const FArguments& InArgs );
+
+protected:
+
+	/* Handler for when the color block is clicked to open the color picker */
+	FReply ColorBlock_OnMouseButtonDown(const FGeometry& MyGeometry, const FPointerEvent& MouseEvent);
+	
+	FReply OnDeleteButtonClicked() const;
+	
+	FReply OnCreateFilterButtonClicked(bool bApplyFilter) const;
+	
+	FReply OnCancelButtonClicked() const;
+
+	FReply OnModifyButtonClicked() const;
+
+	bool CheckFilterValidity() const;
+	
+protected:
+
+	/* True if we are editing a filter, false if we are creating a new filter */
+	bool bInEditMode;
+
+	/* The current filter data we are editing */
+	FCustomTextFilterData FilterData;
+	
+	FOnCreateFilter OnCreateFilter;
+	
+	FOnDeleteFilter OnDeleteFilter;
+	
+	FOnCancelClicked OnCancelClicked;
+
+	FOnModifyFilter OnModifyFilter;
+
+	/* The color block widget that edits the filter color */
+	TSharedPtr<SColorBlock> ColorBlock;
+
+	/* The widget that edits the filter label */
+	TSharedPtr<SEditableTextBox> FilterLabelTextBox;
+};
\ No newline at end of file
diff --git a/Engine/Source/Developer/ToolWidgets/Public/Filters/SFilterSearchBox.h b/Engine/Source/Developer/ToolWidgets/Public/Filters/SFilterSearchBox.h
new file mode 100644
--- /dev/null
+++ b/Engine/Source/Developer/ToolWidgets/Public/Filters/SFilterSearchBox.h
@@ -0,0 +1,120 @@
+// Copyright Epic Games, Inc. All Rights Reserved.
+
+#pragma once
+
+#include "Widgets/SCompoundWidget.h"
+#include "Containers/UnrealString.h"
+#include "Framework/SlateDelegates.h"
+
+class SMenuAnchor;
+class SSearchBox;
+template <typename T> class SListView;
+
+/** A SearchBox widget that contains support for displaying search history.
+ * When used with a FilterBar widget (@see SBasicFilterBar or SFilterBar), it allows the user to save searches from the
+ * history as filter pills
+ */
+class TOOLWIDGETS_API SFilterSearchBox : public SCompoundWidget
+{
+public:
+
+	DECLARE_DELEGATE_OneParam(FOnClickedAddSearchHistoryButton, const FText&)
+
+	SLATE_BEGIN_ARGS(SFilterSearchBox)
+	: _OnTextChanged()
+	, _OnTextCommitted()
+	, _InitialText()
+	, _HintText()
+	, _MaxSearchHistory(5)
+	{}
+
+	/** Invoked whenever the text changes */
+	SLATE_EVENT( FOnTextChanged, OnTextChanged )
+
+	/** Invoked whenever the text is committed (e.g. user presses enter) */
+	SLATE_EVENT( FOnTextCommitted, OnTextCommitted )
+
+	/** Initial text to display for the search text */
+	SLATE_ATTRIBUTE( FText, InitialText )
+
+	/** Hint text to display for the search text when there is no value */
+	SLATE_ATTRIBUTE( FText, HintText )
+
+	/** Whether we should show a dropdown containing the last few searches */
+	SLATE_ATTRIBUTE(bool, ShowSearchHistory)
+
+	/** The maximum number of items to show in the Search History */
+	SLATE_ARGUMENT(int32, MaxSearchHistory)
+
+	SLATE_END_ARGS()
+
+	/** Constructs this widget with InArgs */
+	void Construct( const FArguments& InArgs );
+
+	/** Sets the text string currently being edited */
+	void SetText(const TAttribute< FText >& InNewText);
+
+	/** Get the text string currently being edited */
+	FText GetText() const;
+
+	/** Set or clear the current error reporting information for this search box */
+	void SetError( const FText& InError );
+	void SetError( const FString& InError );
+
+	/** SWidget interface */
+	void OnFocusLost(const FFocusEvent& InFocusEvent) override;
+
+	void SetOnClickedAddSearchHistoryButton(FOnClickedAddSearchHistoryButton InOnClickedAddSearchHistoryButton)
+	{
+		OnClickedAddSearchHistoryButton = InOnClickedAddSearchHistoryButton;
+	}
+
+private:
+	
+	/** Handler for when text in the editable text box changed */
+	void HandleTextChanged(const FText& NewText);
+
+	/** Handler for when text in the editable text box changed */
+	void HandleTextCommitted(const FText& NewText, ETextCommit::Type CommitType);
+
+	/** Add the given search text to the search history */
+	void UpdateSearchHistory(const FText &NewText);
+
+	/** Called by SListView when the selection changes in the search history list */
+	void OnSelectionChanged( TSharedPtr<FText> NewValue, ESelectInfo::Type SelectInfo );
+
+	/** Makes the widget for a single search history row in the list view */
+	TSharedRef<ITableRow> MakeSearchHistoryRowWidget(TSharedPtr<FText> SearchText, const TSharedRef<STableViewBase>& OwnerTable);
+
+	/** Handler for when the chevron to open the search history is clicked */
+	FReply OnClickedSearchHistory();
+
+private:
+
+	/** The actual Search Box */
+	TSharedPtr<SSearchBox> SearchBox;
+	
+	/** Delegate for when text is changed in the edit box */
+	FOnTextChanged OnTextChanged;
+
+	/** Delegate for when text is changed in the edit box */
+	FOnTextCommitted OnTextCommitted;
+
+	/** Delegate for when the Plus icon next to an item in the search history is clicked */
+	FOnClickedAddSearchHistoryButton OnClickedAddSearchHistoryButton;
+
+	/** The max amount of items to show in the Search History */
+	int32 MaxSearchHistory;
+
+	/** The array containing the last few searched items */
+	TArray<TSharedPtr<FText>> SearchHistory;
+
+	/** Search History elements */
+	TSharedPtr< SMenuAnchor > SearchHistoryBox;
+
+	/** The ListView containing our search history */
+	TSharedPtr<SListView< TSharedPtr<FText>>> SearchHistoryListView;
+
+	/** The Text that shows up when there are no items in the search history */
+	TSharedPtr<FText> EmptySearchHistoryText;
+};
diff --git a/Engine/Source/Developer/ToolWidgets/ToolWidgets.Build.cs b/Engine/Source/Developer/ToolWidgets/ToolWidgets.Build.cs
--- a/Engine/Source/Developer/ToolWidgets/ToolWidgets.Build.cs
+++ b/Engine/Source/Developer/ToolWidgets/ToolWidgets.Build.cs
@@ -23,6 +23,7 @@ public class ToolWidgets : ModuleRules
 
 		PrivateDependencyModuleNames.AddRange(
 			new string[] {
+				"AppFramework",
 				"ApplicationCore"
 			}
 		);
diff --git a/Engine/Source/Editor/EditorWidgets/Public/Filters/SAssetFilterBar.h b/Engine/Source/Editor/EditorWidgets/Public/Filters/SAssetFilterBar.h
--- a/Engine/Source/Editor/EditorWidgets/Public/Filters/SAssetFilterBar.h
+++ b/Engine/Source/Editor/EditorWidgets/Public/Filters/SAssetFilterBar.h
@@ -78,6 +78,7 @@ public:
 	
 	using FOnFilterChanged = typename SBasicFilterBar<FilterType>::FOnFilterChanged;
 	using FOnExtendAddFilterMenu = typename SBasicFilterBar<FilterType>::FOnExtendAddFilterMenu;
+	using FCreateTextFilter = typename SBasicFilterBar<FilterType>::FCreateTextFilter;
 
 	SLATE_BEGIN_ARGS( SAssetFilterBar<FilterType> )
 		: _UseDefaultAssetFilters(true)
@@ -97,6 +98,17 @@ public:
 		/** Initial List of Custom Class filters that will be added to the AddFilter Menu */
 		SLATE_ARGUMENT( TArray<TSharedRef<FCustomClassFilterData>>, CustomClassFilters)
 
+		/** A delegate to create a TTextFilter for FilterType items. If provided, will allow creation of custom text filters
+		 *  from the filter dropdown menu.
+		 */
+		SLATE_ARGUMENT(FCreateTextFilter, CreateTextFilter)
+			
+		/** An SFilterSearchBox that can be attached to this filter bar. When provided along with a CreateTextFilter
+		 *  delegate, allows the user to save searches from the Search Box as text filters for the filter bar.
+		 *	NOTE: Will bind a delegate to SFilterSearchBox::OnClickedAddSearchHistoryButton
+		 */
+		SLATE_ARGUMENT(TSharedPtr<SFilterSearchBox>, FilterSearchBox)
+
 		/** Whether the filter bar should provide the default asset filters */
 		SLATE_ARGUMENT(bool, UseDefaultAssetFilters)
 	
@@ -112,6 +124,8 @@ public:
 		Args._OnFilterChanged = InArgs._OnFilterChanged;
 		Args._CustomFilters = InArgs._CustomFilters;
 		Args._OnExtendAddFilterMenu = InArgs._OnExtendAddFilterMenu;
+		Args._CreateTextFilter = InArgs._CreateTextFilter;
+		Args._FilterSearchBox = InArgs._FilterSearchBox;
 		
 		SBasicFilterBar<FilterType>::Construct(Args);
 		
diff --git a/Engine/Source/Editor/EditorWidgets/Public/Filters/SFilterBar.h b/Engine/Source/Editor/EditorWidgets/Public/Filters/SFilterBar.h
--- a/Engine/Source/Editor/EditorWidgets/Public/Filters/SFilterBar.h
+++ b/Engine/Source/Editor/EditorWidgets/Public/Filters/SFilterBar.h
@@ -42,6 +42,7 @@ public:
 	using FConvertItemToAssetData = typename FAssetFilter<FilterType>::FConvertItemToAssetData;
 	using FCompareItemWithClassNames = typename FAssetFilter<FilterType>::FCompareItemWithClassNames;
 	using FOnExtendAddFilterMenu = typename SBasicFilterBar<FilterType>::FOnExtendAddFilterMenu;
+	using FCreateTextFilter = typename SBasicFilterBar<FilterType>::FCreateTextFilter;
 	
 	SLATE_BEGIN_ARGS( SFilterBar<FilterType> )
 		: _UseDefaultAssetFilters(true)
@@ -69,6 +70,17 @@ public:
 	
 		/** Whether the filter bar should provide the default asset filters */
 		SLATE_ARGUMENT(bool, UseDefaultAssetFilters)
+
+		/** A delegate to create a TTextFilter for FilterType items. If provided, will allow creation of custom text filters
+		 *  from the filter dropdown menu.
+		 */
+		SLATE_ARGUMENT(FCreateTextFilter, CreateTextFilter)
+			
+		/** An SFilterSearchBox that can be attached to this filter bar. When provided along with a CreateTextFilter
+		 *  delegate, allows the user to save searches from the Search Box as text filters for the filter bar.
+		 *	NOTE: Will bind a delegate to SFilterSearchBox::OnClickedAddSearchHistoryButton
+		 */
+		SLATE_ARGUMENT(TSharedPtr<SFilterSearchBox>, FilterSearchBox)
 	
 	SLATE_END_ARGS()
 
@@ -81,6 +93,8 @@ public:
 		Args._UseDefaultAssetFilters = InArgs._UseDefaultAssetFilters;
 		Args._CustomClassFilters = InArgs._CustomClassFilters;
 		Args._OnExtendAddFilterMenu = InArgs._OnExtendAddFilterMenu;
+		Args._CreateTextFilter = InArgs._CreateTextFilter;
+		Args._FilterSearchBox = InArgs._FilterSearchBox;
 		
 		SAssetFilterBar<FilterType>::Construct(Args);
 
