SBasicFilterBar: Fix non unity compiler error due to missing include

diff --git a/Engine/Source/Developer/ToolWidgets/Public/Filters/SBasicFilterBar.h b/Engine/Source/Developer/ToolWidgets/Public/Filters/SBasicFilterBar.h
--- a/Engine/Source/Developer/ToolWidgets/Public/Filters/SBasicFilterBar.h
+++ b/Engine/Source/Developer/ToolWidgets/Public/Filters/SBasicFilterBar.h
@@ -22,6 +22,7 @@
 #include "SSimpleButton.h"
 #include "Styling/StyleColors.h"
 #include "Filters/SFilterSearchBox.h"
+#include "Framework/Application/SlateApplication.h"
 
 #include "SBasicFilterBar.generated.h"
 
