SFilterSearchBox: Fix duplicate localization key

diff --git a/Engine/Source/Developer/ToolWidgets/Private/Filters/SFilterSearchBox.cpp b/Engine/Source/Developer/ToolWidgets/Private/Filters/SFilterSearchBox.cpp
--- a/Engine/Source/Developer/ToolWidgets/Private/Filters/SFilterSearchBox.cpp
+++ b/Engine/Source/Developer/ToolWidgets/Private/Filters/SFilterSearchBox.cpp
@@ -186,7 +186,7 @@ TSharedRef<ITableRow> SFilterSearchBox::MakeSearchHistoryRowWidget(TSharedPtr<FT
 			{
 				return this->OnClickedAddSearchHistoryButton.IsBound() ? EVisibility::Visible : EVisibility::Collapsed;
 			})
-			.ToolTipText(LOCTEXT("SearchHistoryToolTipText", "Save this search as a custom filter"))
+			.ToolTipText(LOCTEXT("SaveSearchHistoryToolTipText", "Save this search as a custom filter"))
 			[
 				SNew(SImage)
 				.Image(FAppStyle::Get().GetBrush("Icons.PlusCircle"))
