SCustomTextFilterDialog: Fix missing header guard

diff --git a/Engine/Source/Developer/ToolWidgets/Public/Filters/SCustomTextFilterDialog.h b/Engine/Source/Developer/ToolWidgets/Public/Filters/SCustomTextFilterDialog.h
--- a/Engine/Source/Developer/ToolWidgets/Public/Filters/SCustomTextFilterDialog.h
+++ b/Engine/Source/Developer/ToolWidgets/Public/Filters/SCustomTextFilterDialog.h
@@ -1,5 +1,7 @@
 // Copyright Epic Games, Inc. All Rights Reserved.
 
+#pragma once
+
 #include "Widgets/SCompoundWidget.h"
 #include "Filters/FilterBase.h"
 #include "Core/Public/Misc/TextFilter.h"
