Change SCustomDialog and SMessageDialog IconBrush() accepting Fname to Icon() accepting any slate brush

diff --git a/Engine/Plugins/Developer/Concert/ConcertApp/MultiUserServer/Source/MultiUserServer/Private/Widgets/Browser/SConcertServerSessionBrowser.cpp b/Engine/Plugins/Developer/Concert/ConcertApp/MultiUserServer/Source/MultiUserServer/Private/Widgets/Browser/SConcertServerSessionBrowser.cpp
--- a/Engine/Plugins/Developer/Concert/ConcertApp/MultiUserServer/Source/MultiUserServer/Private/Widgets/Browser/SConcertServerSessionBrowser.cpp
+++ b/Engine/Plugins/Developer/Concert/ConcertApp/MultiUserServer/Source/MultiUserServer/Private/Widgets/Browser/SConcertServerSessionBrowser.cpp
@@ -88,7 +88,7 @@ void SConcertServerSessionBrowser::DeleteSessionsWithFakeModalQuestion(const FTe
 	};
 	const TSharedRef<SMessageDialog> Dialog = SNew(SMessageDialog)
 		.Title(LOCTEXT("DisconnectUsersTitle", "Delete session?"))
-		.IconBrush("Icons.WarningWithColor.Large")
+		.Icon(FAppStyle::Get().GetBrush("Icons.WarningWithColor.Large"))
 		.Message(Message)
 		.UseScrollBox(false)
 		.Buttons({
diff --git a/Engine/Plugins/Developer/Concert/ConcertApp/MultiUserServer/Source/MultiUserServer/Private/Widgets/Clients/Browser/SConcertClientBrowser.cpp b/Engine/Plugins/Developer/Concert/ConcertApp/MultiUserServer/Source/MultiUserServer/Private/Widgets/Clients/Browser/SConcertClientBrowser.cpp
--- a/Engine/Plugins/Developer/Concert/ConcertApp/MultiUserServer/Source/MultiUserServer/Private/Widgets/Clients/Browser/SConcertClientBrowser.cpp
+++ b/Engine/Plugins/Developer/Concert/ConcertApp/MultiUserServer/Source/MultiUserServer/Private/Widgets/Clients/Browser/SConcertClientBrowser.cpp
@@ -170,7 +170,7 @@ TSharedRef<SWidget> UE::MultiUserServer::SConcertClientBrowser::CreateKeepDiscon
 				
 				const TSharedRef<SMessageDialog> Dialog = SNew(SMessageDialog)
 					.Title(LOCTEXT("RemoveDisconnectedClients.Title", "Remove disconnected clients?"))
-					.IconBrush("Icons.WarningWithColor.Large")
+					.Icon(FAppStyle::Get().GetBrush("Icons.WarningWithColor.Large"))
 					.Message(FText::Format(LOCTEXT("RemoveDisconnectedClients.MessageFmt", "There are {0} disconnected clients. If you proceed, these clients will be removed from the session browser; opened log tabs will remain open.\nProceed?"), NumDisconnected))
 					.UseScrollBox(false)
 					.Buttons({
diff --git a/Engine/Source/Developer/ToolWidgets/Private/Dialog/SCustomDialog.cpp b/Engine/Source/Developer/ToolWidgets/Private/Dialog/SCustomDialog.cpp
--- a/Engine/Source/Developer/ToolWidgets/Private/Dialog/SCustomDialog.cpp
+++ b/Engine/Source/Developer/ToolWidgets/Private/Dialog/SCustomDialog.cpp
@@ -2,8 +2,6 @@
 
 #include "Dialog/SCustomDialog.h"
 
-#include "HAL/PlatformApplicationMisc.h"
-
 #include "Framework/Application/SlateApplication.h"
 #include "Framework/Docking/TabManager.h"
 #include "Logging/LogMacros.h"
@@ -12,20 +10,31 @@
 #include "Widgets/Images/SImage.h"
 #include "Widgets/Input/SButton.h"
 #include "Widgets/Text/STextBlock.h"
-#include "Widgets/Layout/SSpacer.h"
 #include "Widgets/Layout/SBox.h"
 #include "Widgets/Layout/SScrollBox.h"
 #include "Widgets/Layout/SUniformGridPanel.h"
 #include "Widgets/SBoxPanel.h"
+#include "Widgets/SNullWidget.h"
 
 DEFINE_LOG_CATEGORY_STATIC(LogCustomDialog, Log, All);
 
+SCustomDialog::FArguments& SCustomDialog::FArguments::IconBrush(FName InIconBrush)
+{
+	const FSlateBrush* ImageBrush = FAppStyle::Get().GetBrush(InIconBrush);
+	if (ensureMsgf(ImageBrush != nullptr, TEXT("Brush %s is unknown"), *InIconBrush.ToString()))
+	{
+		_Icon = ImageBrush;
+	}
+	return *this;
+}
+
 void SCustomDialog::Construct(const FArguments& InArgs)
 {
 	UE_LOG(LogCustomDialog, Log, TEXT("Dialog displayed:"), *InArgs._Title.ToString());
 	check(InArgs._Buttons.Num() > 0);
 	
 	OnClosed = InArgs._OnClosed;
+	bAutoCloseOnButtonPress = InArgs._AutoCloseOnButtonPress;
 	
 	SWindow::Construct( SWindow::FArguments(InArgs._WindowArguments)
 		.Title(InArgs._Title)
@@ -74,23 +83,19 @@ TSharedRef<SWidget> SCustomDialog::CreateContentBox(const FArguments& InArgs)
 {
 	TSharedRef<SHorizontalBox> ContentBox = SNew(SHorizontalBox);
 
-	if (InArgs._IconBrush.IsValid())
-	{
-		const FSlateBrush* ImageBrush = FAppStyle::Get().GetBrush(InArgs._IconBrush);
-		if (ensureMsgf(ImageBrush != nullptr, TEXT("Brush %s is unknown"), *InArgs._IconBrush.ToString()))
-		{
-			ContentBox->AddSlot()
-				.AutoWidth()
-				.VAlign(InArgs._VAlignIcon)
-				.HAlign(InArgs._HAlignIcon)
-				.Padding(0, 0, 8, 0)
-				[
-					SNew(SImage)
-					.DesiredSizeOverride(InArgs._IconDesiredSizeOverride)
-					.Image(ImageBrush)
-				];
-		}
-	}
+	ContentBox->AddSlot()
+		.AutoWidth()
+		.VAlign(InArgs._VAlignIcon)
+		.HAlign(InArgs._HAlignIcon)
+		.Padding(0, 0, 8, 0)
+		[
+			!InArgs._Icon.IsSet() || !InArgs._Icon.Get()
+				// If no icon is specified, we still want to pad the content towards the right
+				? SNullWidget::NullWidget
+				: SNew(SImage)
+				.DesiredSizeOverride(InArgs._IconDesiredSizeOverride)
+				.Image(InArgs._Icon)
+		];
 
 	if (InArgs._UseScrollBox)
 	{
@@ -202,7 +207,10 @@ FReply SCustomDialog::OnButtonClicked(FSimpleDelegate OnClicked, int32 ButtonInd
 {
 	LastPressedButton = ButtonIndex;
 
-	FSlateApplication::Get().RequestDestroyWindow(StaticCastSharedRef<SWindow>(this->AsShared()));
+	if (bAutoCloseOnButtonPress)
+	{
+		RequestDestroyWindow();
+	}
 
 	OnClicked.ExecuteIfBound();
 	return FReply::Handled();
diff --git a/Engine/Source/Developer/ToolWidgets/Private/Dialog/SMessageDialog.cpp b/Engine/Source/Developer/ToolWidgets/Private/Dialog/SMessageDialog.cpp
--- a/Engine/Source/Developer/ToolWidgets/Private/Dialog/SMessageDialog.cpp
+++ b/Engine/Source/Developer/ToolWidgets/Private/Dialog/SMessageDialog.cpp
@@ -6,12 +6,7 @@
 
 #include "Styling/AppStyle.h"
 #include "Styling/SlateBrush.h"
-#include "Widgets/Images/SImage.h"
-#include "Widgets/Input/SButton.h"
 #include "Widgets/Text/STextBlock.h"
-#include "Widgets/Layout/SSpacer.h"
-#include "Widgets/Layout/SScrollBox.h"
-#include "Widgets/SBoxPanel.h"
 #include "Widgets/Input/SHyperlink.h"
 
 void SMessageDialog::Construct(const FArguments& InArgs)
@@ -30,7 +25,8 @@ void SMessageDialog::Construct(const FArguments& InArgs)
 		.WindowArguments(InArgs._WindowArguments)
 		.RootPadding(16.f)
 		.Buttons(InArgs._Buttons)
-		.IconBrush(InArgs._IconBrush)
+		.AutoCloseOnButtonPress(InArgs._AutoCloseOnButtonPress)
+		.Icon(InArgs._Icon)
 		.HAlignContent(HAlign_Fill)
 		.VAlignContent(VAlign_Fill)
 		.IconDesiredSizeOverride(FVector2D(24.f, 24.f))
diff --git a/Engine/Source/Developer/ToolWidgets/Public/Dialog/SCustomDialog.h b/Engine/Source/Developer/ToolWidgets/Public/Dialog/SCustomDialog.h
--- a/Engine/Source/Developer/ToolWidgets/Public/Dialog/SCustomDialog.h
+++ b/Engine/Source/Developer/ToolWidgets/Public/Dialog/SCustomDialog.h
@@ -3,6 +3,7 @@
 #pragma once
 
 #include "CoreMinimal.h"
+#include "Styling/AppStyle.h"
 #include "Widgets/SWindow.h"
 
 /**
@@ -56,7 +57,9 @@ public:
 	};
 
 	SLATE_BEGIN_ARGS(SCustomDialog) 
-		:  _HAlignIcon(HAlign_Left)
+		: _AutoCloseOnButtonPress(true)
+		, _Icon(nullptr)
+		, _HAlignIcon(HAlign_Left)
 		, _VAlignIcon(VAlign_Center)
 		, _RootPadding(FMargin(4.f))
 		, _ButtonAreaPadding(FMargin(20.f, 0.f, 0.f, 0.f))
@@ -84,12 +87,14 @@ public:
 
 		/** Provides default values for SWindow::FArguments not overriden by SCustomDialog. */
 		SLATE_ARGUMENT(SWindow::FArguments, WindowArguments)
-	
+
+		/** Whether to automatically close this window when any button is pressed (default: true) */
+		SLATE_ARGUMENT(bool, AutoCloseOnButtonPress)
 
 		/*********** Cosmetic ***********/
 
-		/** Optional icon to display in the dialog (default: none) */
-		SLATE_ARGUMENT(FName, IconBrush)
+		/** Optional icon to display (default: empty, translucent)*/
+		SLATE_ATTRIBUTE(const FSlateBrush*, Icon)
 	
 		/** When specified, ignore the brushes size and report the DesiredSizeOverride as the desired image size (default: use icon size) */
 		SLATE_ATTRIBUTE(TOptional<FVector2D>, IconDesiredSizeOverride)
@@ -130,6 +135,10 @@ public:
 
 
 		/********** Legacy - do not use **********/
+		
+		/** Optional icon to display in the dialog (default: none) */
+		UE_DEPRECATED(5.1, "Use Icon() instead")
+		FArguments& IconBrush(FName InIconBrush);
 	
 		/** Content for the dialog (deprecated - use Content instead)*/
 		SLATE_ARGUMENT_DEPRECATED(TSharedPtr<SWidget>, DialogContent, 5.1, "Use Content() instead.")
@@ -159,4 +168,6 @@ private:
 	int32 LastPressedButton = -1;
 
 	FSimpleDelegate OnClosed;
+
+	bool bAutoCloseOnButtonPress;
 };
diff --git a/Engine/Source/Developer/ToolWidgets/Public/Dialog/SMessageDialog.h b/Engine/Source/Developer/ToolWidgets/Public/Dialog/SMessageDialog.h
--- a/Engine/Source/Developer/ToolWidgets/Public/Dialog/SMessageDialog.h
+++ b/Engine/Source/Developer/ToolWidgets/Public/Dialog/SMessageDialog.h
@@ -16,7 +16,9 @@ public:
 	using FButton = SCustomDialog::FButton;
 
 	SLATE_BEGIN_ARGS(SMessageDialog)
-		:  _WrapMessageAt(512.f)
+		: _AutoCloseOnButtonPress(true)
+		, _Icon(nullptr)
+		, _WrapMessageAt(512.f)
 	{}
 
 		/********** Functional **********/
@@ -36,10 +38,13 @@ public:
 		/** Provides default values for SWindow::FArguments not overriden by SCustomDialog. */
 		SLATE_ARGUMENT(SWindow::FArguments, WindowArguments)
 	
+		/** Whether to automatically close this window when any button is pressed (default: true) */
+		SLATE_ARGUMENT(bool, AutoCloseOnButtonPress)
+	
 		/********** Cosmetic **********/
 	
-		/** Optional icon to display in the dialog. (default: none) */
-		SLATE_ARGUMENT(FName, IconBrush)
+		/** Optional icon to display in the dialog. (default: empty) */
+		SLATE_ARGUMENT(const FSlateBrush*, Icon)
 
 		/** Should this dialog use a scroll box for over-sized content? (default: true) */
 		SLATE_ARGUMENT(bool, UseScrollBox)
diff --git a/Engine/Source/Editor/Kismet/Private/BlueprintEditor.cpp b/Engine/Source/Editor/Kismet/Private/BlueprintEditor.cpp
--- a/Engine/Source/Editor/Kismet/Private/BlueprintEditor.cpp
+++ b/Engine/Source/Editor/Kismet/Private/BlueprintEditor.cpp
@@ -4031,7 +4031,7 @@ void FBlueprintEditor::DeleteUnusedVariables_OnClicked()
 
 		TSharedRef<SCustomDialog> CustomDialog = SNew(SCustomDialog)
 			.Title(LOCTEXT("DeleteUnusedVariablesDialog_Title", "Delete Unused Variables"))
-			.IconBrush("NotificationList.DefaultMessage")
+			.Icon(FAppStyle::Get().GetBrush("NotificationList.DefaultMessage"))
 			.Content()
 			[
 				DialogContain
diff --git a/Engine/Source/Editor/UnrealEd/Private/PlayLevel.cpp b/Engine/Source/Editor/UnrealEd/Private/PlayLevel.cpp
--- a/Engine/Source/Editor/UnrealEd/Private/PlayLevel.cpp
+++ b/Engine/Source/Editor/UnrealEd/Private/PlayLevel.cpp
@@ -1569,7 +1569,7 @@ static bool ShowBlueprintErrorDialog( TArray<UBlueprint*> ErroredBlueprints )
 
 	CustomDialog = SNew(SCustomDialog)
 		.Title(DialogTitle)
-		.IconBrush("NotificationList.DefaultMessage")
+		.Icon(FAppStyle::Get().GetBrush("NotificationList.DefaultMessage"))
 		.Content()
 		[
 			DialogContents
