Fix deprecated function not getting API exported

diff --git a/Engine/Source/Developer/ToolWidgets/Public/Dialog/SCustomDialog.h b/Engine/Source/Developer/ToolWidgets/Public/Dialog/SCustomDialog.h
--- a/Engine/Source/Developer/ToolWidgets/Public/Dialog/SCustomDialog.h
+++ b/Engine/Source/Developer/ToolWidgets/Public/Dialog/SCustomDialog.h
@@ -138,7 +138,7 @@ public:
 		
 		/** Optional icon to display in the dialog (default: none) */
 		UE_DEPRECATED(5.1, "Use Icon() instead")
-		FArguments& IconBrush(FName InIconBrush);
+		TOOLWIDGETS_API FArguments& IconBrush(FName InIconBrush);
 	
 		/** Content for the dialog (deprecated - use Content instead)*/
 		SLATE_ARGUMENT_DEPRECATED(TSharedPtr<SWidget>, DialogContent, 5.1, "Use Content() instead.")
