Content Browser: Add an alternate vertical filter layout
Remove the overflow menu and add wrapping for fiters in the horizontal layout

diff --git a/Engine/Source/Developer/ToolWidgets/Private/Filters/SFilterSearchBox.cpp b/Engine/Source/Developer/ToolWidgets/Private/Filters/SFilterSearchBox.cpp
--- a/Engine/Source/Developer/ToolWidgets/Private/Filters/SFilterSearchBox.cpp
+++ b/Engine/Source/Developer/ToolWidgets/Private/Filters/SFilterSearchBox.cpp
@@ -73,6 +73,7 @@ void SFilterSearchBox::Construct( const FArguments& InArgs )
 
 				+SHorizontalBox::Slot()
 				.HAlign(HAlign_Right)
+				.VAlign(VAlign_Center)
 				.Padding(0, 0, 4, 0)
 				[
 					// Chevron to open the search history dropdown
diff --git a/Engine/Source/Developer/ToolWidgets/Public/Filters/SBasicFilterBar.h b/Engine/Source/Developer/ToolWidgets/Public/Filters/SBasicFilterBar.h
--- a/Engine/Source/Developer/ToolWidgets/Public/Filters/SBasicFilterBar.h
+++ b/Engine/Source/Developer/ToolWidgets/Public/Filters/SBasicFilterBar.h
@@ -12,12 +12,15 @@
 #include "Textures/SlateIcon.h"
 #include "Widgets/Layout/SBorder.h"
 #include "ToolMenus.h"
-#include "Widgets/Images/SImage.h"
+#include "Widgets/Images/SLayeredImage.h"
 #include "Widgets/SBoxPanel.h"
 #include "Misc/FilterCollection.h"
 #include "UObject/Object.h"
 #include "Core/Public/Misc/TextFilter.h"
 #include "Widgets/SWindow.h"
+#include "Widgets/Layout/SWrapBox.h"
+#include "Widgets/Layout/SScrollBox.h"
+#include "Widgets/Layout/SWidgetSwitcher.h"
 #include "Filters/SCustomTextFilterDialog.h"
 #include "SSimpleButton.h"
 #include "Styling/StyleColors.h"
@@ -41,6 +44,14 @@ public:
 	FOnPopulateAddFilterMenu PopulateFilterMenu;
 };
 
+/** Describes if the filters are laid out horizontally (ScrollBox) or vertically (WrapBox) */
+UENUM()
+enum class EFilterBarLayout : uint8
+{
+	Horizontal,
+	Vertical
+};
+
 /** Forward declaration for a helper struct used to activate external filters */
 template<typename FilterType>
 struct FFrontendFilterExternalActivationHelper;
@@ -70,7 +81,10 @@ public:
 	/** Delegate to create a TTextFilter used to compare FilterType with text queries */
     DECLARE_DELEGATE_RetVal( TSharedPtr<ICustomTextFilter<FilterType>>, FCreateTextFilter);
 	
- 	SLATE_BEGIN_ARGS( SBasicFilterBar<FilterType> ){}
+ 	SLATE_BEGIN_ARGS( SBasicFilterBar<FilterType> )
+ 		: _FilterBarLayout(EFilterBarLayout::Horizontal)
+		, _CanChangeOrientation(false)
+	{}
 
  		/** Delegate for when filters have changed */
  		SLATE_EVENT( FOnFilterChanged, OnFilterChanged )
@@ -92,6 +106,12 @@ public:
 		 *	NOTE: Will bind a delegate to SFilterSearchBox::OnSaveSearchClicked
 		 */
         SLATE_ARGUMENT(TSharedPtr<SFilterSearchBox>, FilterSearchBox)
+
+		/** The layout that determines how the filters are laid out */
+		SLATE_ARGUMENT(EFilterBarLayout, FilterBarLayout)
+	
+		/** If true, allow dynamically changing the orientation and saving in the config */
+		SLATE_ARGUMENT(bool, CanChangeOrientation)
 	
  	SLATE_END_ARGS()
 
@@ -100,6 +120,36 @@ public:
 		OnFilterChanged = InArgs._OnFilterChanged;
  		OnExtendAddFilterMenu = InArgs._OnExtendAddFilterMenu;
  		CreateTextFilter = InArgs._CreateTextFilter;
+ 		FilterBarLayout = InArgs._FilterBarLayout;
+ 		bCanChangeOrientation = InArgs._CanChangeOrientation;
+
+ 		// We use a widgetswitcher to allow dynamically swapping between the layouts
+ 		ChildSlot
+		[
+			SAssignNew(FilterBox, SWidgetSwitcher)
+			+ SWidgetSwitcher::Slot()
+			[
+				SAssignNew(HorizontalFilterBox, SWrapBox)
+				.UseAllottedSize(true)
+			]
+			+ SWidgetSwitcher::Slot()
+			[
+				SAssignNew(VerticalFilterBox, SScrollBox)
+				.Visibility_Lambda([this]()
+				{
+					return HasAnyFilters() ? EVisibility::Visible : EVisibility::Collapsed;
+				})
+			]
+		];
+ 		
+ 		if(FilterBarLayout == EFilterBarLayout::Horizontal)
+ 		{
+ 			FilterBox->SetActiveWidget(HorizontalFilterBox.ToSharedRef());
+ 		}
+ 		else
+ 		{
+ 			FilterBox->SetActiveWidget(VerticalFilterBox.ToSharedRef());
+ 		}
 
  		AttachFilterSearchBox(InArgs._FilterSearchBox);
 
@@ -136,9 +186,52 @@ public:
  		SetFrontendFilterActive(InFilter, false);
  	}
 
+	EFilterBarLayout GetFilterLayout()
+ 	{
+ 		return FilterBarLayout;
+ 	}
+	
+	virtual void SetFilterLayout(EFilterBarLayout InFilterBarLayout)
+ 	{
+ 		if(!bCanChangeOrientation)
+ 		{
+ 			return;
+ 		}
+ 		
+ 		FilterBarLayout = InFilterBarLayout;
+ 		
+ 		if(FilterBarLayout == EFilterBarLayout::Horizontal)
+ 		{
+			VerticalFilterBox->ClearChildren();
+ 			
+ 			FilterBox->SetActiveWidget(HorizontalFilterBox.ToSharedRef());
+ 		}
+        else
+        {
+        	HorizontalFilterBox->ClearChildren();
+        	
+        	FilterBox->SetActiveWidget(VerticalFilterBox.ToSharedRef());
+        }
+
+ 		// Add all the filters to the new layout
+ 		for(TSharedRef<SFilter> Filter: Filters)
+ 		{
+ 			AddWidgetToLayout(Filter);
+ 		}
+
+ 		this->Invalidate(EInvalidateWidgetReason::Layout);
+ 	}
+
 	/** Makes the button that summons the Add Filter dropdown on click */
 	static TSharedRef<SWidget> MakeAddFilterButton(TSharedRef<SBasicFilterBar<FilterType>> InFilterBar)
  	{
+ 		TSharedPtr<SLayeredImage> FilterImage = SNew(SLayeredImage)
+		 .Image(FAppStyle::Get().GetBrush("Icons.Filter"))
+		 .ColorAndOpacity(FSlateColor::UseForeground());
+
+ 		// Badge the filter icon if there are filters active
+ 		FilterImage->AddLayer(TAttribute<const FSlateBrush*>(InFilterBar, &SBasicFilterBar<FilterType>::GetFilterBadgeIcon));
+ 		
  		return SNew(SComboButton)
 				.ComboButtonStyle(&FAppStyle::Get().GetWidgetStyle<FComboButtonStyle>("SimpleComboButtonWithIcon"))
 				.ForegroundColor(FSlateColor::UseStyle())
@@ -147,9 +240,7 @@ public:
 				.ContentPadding(FMargin(1, 0))
 				.ButtonContent()
 				[
-					SNew(SImage)
-					.Image(FAppStyle::Get().GetBrush("Icons.Filter"))
-					.ColorAndOpacity(FSlateColor::UseForeground())
+					FilterImage.ToSharedRef()
 				];
  	}
 
@@ -672,15 +763,12 @@ public:
  				{
  					SetFrontendFilterActive(FrontendFilter.ToSharedRef(), false); // Deactivate.
  				}
+
+ 				RemoveWidgetFromLayout(FilterToRemove);
  			}
-		
- 			ChildSlot
-			 [
-				 SNullWidget::NullWidget
-			 ];
 
  			Filters.Empty();
-
+ 			
  			// Notify that a filter has changed
  			OnFilterChanged.ExecuteIfBound();
  		}
@@ -732,6 +820,20 @@ public:
  	}
 
 protected:
+
+	/** Overlay for filter icon, badge if there are any active filters */
+	const FSlateBrush* GetFilterBadgeIcon() const
+	{
+		for (const TSharedRef<SFilter>& Filter : Filters)
+		{
+			if(Filter->IsEnabled())
+			{
+				return FAppStyle::Get().GetBrush("Icons.BadgeModified");
+			}
+		}
+
+		return nullptr;
+	}
 	
 	/** Remove all filters except the specified one */
 	virtual void RemoveAllButThis(const TSharedRef<SFilter>& FilterToKeep)
@@ -756,32 +858,38 @@ protected:
 		OnFilterChanged.ExecuteIfBound();
 	}
 
-	/** Helper that creates a toolbar with all the given SFilter's as toolbar items. Filters that don't fit appear in the overflow menu as toggles. */
-	static TSharedRef<SWidget> MakeFilterToolBarWidget(const TArray<TSharedRef<SFilter>>& Filters)
+	/** Add a widget to the current filter layout */
+	void AddWidgetToLayout(const TSharedRef<SWidget> WidgetToAdd)
 	{
-		FSlimHorizontalToolBarBuilder ToolbarBuilder(TSharedPtr<const FUICommandList>(), FMultiBoxCustomization::None, TSharedPtr<FExtender>(), true);
-		ToolbarBuilder.SetLabelVisibility(EVisibility::Collapsed);
-		ToolbarBuilder.SetStyle(&FAppStyle::Get(), "ContentBrowser.FilterToolBar");
-
-		for (const TSharedRef<SFilter>& Filter : Filters)
+		if(FilterBarLayout == EFilterBarLayout::Horizontal)
 		{
-			ToolbarBuilder.AddWidget(Filter, NAME_None, true, EHorizontalAlignment::HAlign_Fill, FNewMenuDelegate::CreateLambda([Filter](FMenuBuilder& MenuBuilder)
-			{
-				FUIAction Action;
-				Action.GetActionCheckState = FGetActionCheckState::CreateLambda([Filter]()
-				{
-					return Filter->IsEnabled() ? ECheckBoxState::Checked : ECheckBoxState::Unchecked;
-				});
-				Action.ExecuteAction = FExecuteAction::CreateLambda([Filter]()
-				{
-					Filter->SetEnabled(!Filter->IsEnabled());
-				});
-
-				MenuBuilder.AddMenuEntry(Filter->GetFilterDisplayName(), FText::GetEmpty(), FSlateIcon(), Action, NAME_None, EUserInterfaceActionType::ToggleButton);
-			}));
+			HorizontalFilterBox->AddSlot()
+			.Padding(3, 3)
+			[
+				WidgetToAdd
+			];
+		}
+		else
+		{
+			VerticalFilterBox->AddSlot()
+			.Padding(4, 2)
+			[
+				WidgetToAdd
+			];
 		}
+	}
 
-		return ToolbarBuilder.MakeWidget();
+	/** Remove a widget from the current filter layout */
+	void RemoveWidgetFromLayout(const TSharedRef<SWidget> WidgetToRemove)
+	{
+		if(FilterBarLayout == EFilterBarLayout::Horizontal)
+		{
+			HorizontalFilterBox->RemoveSlot(WidgetToRemove);
+		}
+		else
+		{
+			VerticalFilterBox->RemoveSlot(WidgetToRemove);
+		}
 	}
 	
  	/** Sets the active state of a frontend filter. */
@@ -828,10 +936,7 @@ protected:
 	{
 		Filters.Add(FilterToAdd);
 	
-		ChildSlot
-		[
-			MakeFilterToolBarWidget(Filters)
-		];
+		AddWidgetToLayout(FilterToAdd);
 	}
 
 	/** Handler for when the enable only this button was clicked on a single filter */
@@ -888,10 +993,7 @@ protected:
 			OnFilterChanged.ExecuteIfBound();
 		}
 
-		ChildSlot
-		[
-			MakeFilterToolBarWidget(Filters)
-		];
+		RemoveWidgetFromLayout(FilterToRemove);
 	}
 
 	/** Remove a filter from the filter bar */
@@ -1081,8 +1183,9 @@ protected:
  	}
  	
  	/** Handler for when a checkbox next to a custom text filter is clicked */
-	void CustomTextFilterClicked(ECheckBoxState CheckBoxState, TSharedRef<FFilterBase<FilterType>> FrontendFilter)
+	void CustomTextFilterClicked(ECheckBoxState CheckBoxState, TSharedRef<ICustomTextFilter<FilterType>> Filter)
     {
+		TSharedRef<FFilterBase<FilterType>> FrontendFilter = Filter->GetFilter().ToSharedRef();
     	if (CheckBoxState == ECheckBoxState::Unchecked)
     	{
     		RemoveFilter(FrontendFilter);
@@ -1285,7 +1388,7 @@ protected:
 			[
 				SNew(SCheckBox)
 				.Style(FAppStyle::Get(), "Menu.CheckBox")
-				.OnCheckStateChanged(FOnCheckStateChanged::CreateSP(this, &SBasicFilterBar<FilterType>::CustomTextFilterClicked, FrontendFilter))
+				.OnCheckStateChanged(FOnCheckStateChanged::CreateSP(this, &SBasicFilterBar<FilterType>::CustomTextFilterClicked, Filter))
 				.IsChecked(bIsFilterChecked)
 			];
 
@@ -1297,9 +1400,9 @@ protected:
 				.ButtonStyle(FAppStyle::Get(), "SimpleButton")
 				.ForegroundColor(FStyleColors::White)
 				.Text(FrontendFilter->GetDisplayName())
-				.OnClicked_Lambda([this, &FrontendFilter]()
+				.OnClicked_Lambda([this, &Filter]()
 				{
-					this->FrontendFilterClicked(FrontendFilter);
+					this->FrontendFilterClicked(Filter->GetFilter().ToSharedRef());
 					FSlateApplication::Get().DismissAllMenus();
 					return FReply::Handled();
 				})
@@ -1410,7 +1513,7 @@ private:
 
 		return UToolMenus::Get()->GenerateWidget(FilterMenuName, ToolMenuContext);
 	}
-
+	
 	/** Function to populate the add filter menu */
 	void PopulateAddFilterMenu(UToolMenu* Menu)
 	{
@@ -1420,6 +1523,19 @@ private:
 	}
 	
 protected:
+
+	/** The horizontal wrap box which contains all the filters (used in the horizontal layout) */
+	TSharedPtr<SWrapBox> HorizontalFilterBox;
+	
+	/** The vertical scroll box which contains all the filters (used in the vertical layout) */
+	TSharedPtr<SScrollBox> VerticalFilterBox;
+	
+	/** The widget switcher containing the horizontal and vertical filter widgets */
+	TSharedPtr<SWidgetSwitcher> FilterBox;
+
+	/** The current Filter Layout being used */
+	EFilterBarLayout FilterBarLayout;
+	
  	/** All SFilters in the list */
  	TArray< TSharedRef<SFilter> > Filters;
 
@@ -1446,6 +1562,9 @@ protected:
 
 	/** The window containing the custom text filter dialog */
 	TWeakPtr<SWindow> CustomTextFilterWindow;
+
+	/** Whether the orientation can be changed after initailization */
+	bool bCanChangeOrientation;
 	
 	friend struct FFrontendFilterExternalActivationHelper<FilterType>;
 };
diff --git a/Engine/Source/Editor/ContentBrowser/Private/SContentBrowser.cpp b/Engine/Source/Editor/ContentBrowser/Private/SContentBrowser.cpp
--- a/Engine/Source/Editor/ContentBrowser/Private/SContentBrowser.cpp
+++ b/Engine/Source/Editor/ContentBrowser/Private/SContentBrowser.cpp
@@ -701,14 +701,150 @@ TSharedRef<SWidget> SContentBrowser::CreateLockButton(const FContentBrowserConfi
 	return SNullWidget::NullWidget;
 }
 
-TSharedRef<SWidget> SContentBrowser::CreateAssetView(const FContentBrowserConfig* Config)
+void SContentBrowser::OnFilterBarLayoutChanging(EFilterBarLayout NewLayout)
 {
-	// Item height is specified explicitly here, rather than relying on the padding being exactly right,
-	// so that the search box doesn't subtly shift up/down when the filters are wrapped to a new row in the SWrapBox.
-	const float SearchFilterItemHeight = 32.0f;
-
 	const float SearchBoxDesiredWidth = 500.0f;
 
+	if(NewLayout == EFilterBarLayout::Horizontal)
+	{
+		TSharedPtr<SHorizontalBox> SearchBoxSlot =
+			SNew(SHorizontalBox)
+
+			+ SHorizontalBox::Slot()
+			.AutoWidth()
+			.VAlign(VAlign_Top)
+			.Padding(6, 4, 0, 0)
+			[
+				FilterComboButton.ToSharedRef()
+			]
+
+			+ SHorizontalBox::Slot()
+			.FillWidth(1.0)
+			.VAlign(VAlign_Top)
+			.Padding(0, 4, 0, 0)
+			[
+				SNew(SBox)
+				 .VAlign(VAlign_Center)
+				 .WidthOverride(SearchBoxDesiredWidth)
+				 [
+					 SearchBoxPtr.ToSharedRef()	
+				 ]
+			];
+			
+		/** We add the Combo Button and the Search Box to the FilterList itself, so that the filters wrap with them
+		 *	properly in the Horizontal Layout
+		 */
+		FilterListPtr->AddWidgetToCurrentLayout(SearchBoxSlot.ToSharedRef());
+
+		AssetViewBorder->SetContent(
+
+			SNew(SVerticalBox)
+
+			+ SVerticalBox::Slot()
+			.Padding(0.0f)
+			.AutoHeight()
+			[
+				FilterListPtr.ToSharedRef()
+			]
+
+			+ SVerticalBox::Slot()
+			.FillHeight( 1.0f )
+			.Padding( 0, 0 )
+			[
+				AssetViewPtr.ToSharedRef()
+			]
+
+		);
+		
+	}
+	else
+	{
+		AssetViewBorder->SetContent(
+
+			SNew(SSplitter)
+			.PhysicalSplitterHandleSize(2.0f)
+
+			/* Filters in an SScrollBox */
+			+ SSplitter::Slot()
+			.Value(0.12f)
+			[
+				SNew(SVerticalBox)
+				// Don't take up space when there are no filters
+				.Visibility_Lambda([this]
+				{
+					return FilterListPtr->HasAnyFilters() ? EVisibility::Visible : EVisibility::Collapsed;
+				})
+
+				// Header
+				+ SVerticalBox::Slot()
+				.Padding(0.0f, 2.0f)
+				.AutoHeight()
+				[
+					
+					SNew(SBorder)
+					.BorderImage(FAppStyle::Get().GetBrush("Brushes.Header"))
+					.Padding(FMargin(8.0f, 6.0f))
+					.Content()
+					[
+						SNew(STextBlock)
+						.Text(LOCTEXT("FilterListVerticalHeader", "Filters"))
+						.TextStyle(FAppStyle::Get(), "ButtonText")
+						.Font(FAppStyle::Get().GetFontStyle("NormalFontBold"))
+					]
+				]
+
+				// Filter List
+				+ SVerticalBox::Slot()
+				.Padding(0.0f)
+				.FillHeight(1.0f)
+				[
+					FilterListPtr.ToSharedRef()
+				]
+				
+			]
+		
+			+ SSplitter::Slot()
+			.Value(0.88f)
+			[
+				SNew(SVerticalBox)
+
+				+ SVerticalBox::Slot()
+				.Padding(6, 4, 0, 0)
+				.AutoHeight()
+				[
+					SNew(SHorizontalBox)
+
+					+ SHorizontalBox::Slot()
+					.AutoWidth()
+					[
+						FilterComboButton.ToSharedRef()
+					]
+
+					+ SHorizontalBox::Slot()
+					.AutoWidth()
+					[
+						SNew(SBox)
+						 .VAlign(VAlign_Center)
+						 .WidthOverride(SearchBoxDesiredWidth)
+						 [
+							 SearchBoxPtr.ToSharedRef()	
+						 ]
+					]
+				]
+
+				+ SVerticalBox::Slot()
+				.FillHeight( 1.0f )
+				.Padding( 0, 0 )
+				[
+					AssetViewPtr.ToSharedRef()
+				]
+			]
+		);
+	}
+}
+
+TSharedRef<SWidget> SContentBrowser::CreateAssetView(const FContentBrowserConfig* Config)
+{
 	// Create the Filter Bar Widget
 	FilterListPtr = SNew(SFilterList)
 					.OnFilterChanged(this, &SContentBrowser::OnFilterChanged)
@@ -716,6 +852,9 @@ TSharedRef<SWidget> SContentBrowser::CreateAssetView(const FContentBrowserConfig
 					.Visibility((Config != nullptr ? Config->bCanShowFilters : true) ? EVisibility::Visible : EVisibility::Collapsed)
 					.FrontendFilters(FrontendFilters)
 					.FilterBarIdentifier(InstanceName)
+					.FilterBarLayout(EFilterBarLayout::Vertical)
+					.CanChangeOrientation(true)
+					.OnFilterBarLayoutChanging(this, &SContentBrowser::OnFilterBarLayoutChanging)
 					.UseSharedSettings(true)
 					.CreateTextFilter(SFilterList::FCreateTextFilter::CreateLambda([this]
 					{
@@ -729,44 +868,12 @@ TSharedRef<SWidget> SContentBrowser::CreateAssetView(const FContentBrowserConfig
 					.AddMetaData<FTagMetaData>(FTagMetaData(TEXT("ContentBrowserFilters")));
 	
 	// Create the Filter Combo Button
-	TSharedPtr<SWidget> FilterComboButton = SFilterList::MakeAddFilterButton(FilterListPtr.ToSharedRef());
-	
+	FilterComboButton = SFilterList::MakeAddFilterButton(FilterListPtr.ToSharedRef());
 	TSharedPtr<ISlateMetaData> FilterComboButtonMetaData = MakeShared<FTagMetaData>(TEXT("ContentBrowserFiltersCombo"));
 	FilterComboButton->AddMetadata(FilterComboButtonMetaData.ToSharedRef());
 	FilterComboButton->SetVisibility((Config != nullptr ? Config->bCanShowFilters : true) ? EVisibility::Visible : EVisibility::Collapsed);
 
-	return
-		SNew(SVerticalBox)
-		+ SVerticalBox::Slot()
-		.Padding(0.0f)
-		.AutoHeight()
-		[
-			SNew(SBorder)
-			.BorderImage(FAppStyle::Get().GetBrush("Brushes.Panel"))
-			.Padding(FMargin(2.0f, 2.0f, 2.0f, 0.0f))
-			[
-				SNew(SWrapBox)
-				.UseAllottedSize(true)
-				.InnerSlotPadding(FVector2D(4.0f, 0.0f))
-				+ SWrapBox::Slot()
-				.HAlign(HAlign_Left)
-				.VAlign(VAlign_Center)
-				.FillEmptySpace(true)
-				[
-					// Search field / Save Search button / Filter dropdown
-					SNew(SBox)
-					.MinDesiredHeight(SearchFilterItemHeight)
-					.VAlign(VAlign_Center)
-					.WidthOverride(SearchBoxDesiredWidth)
-					[
-						SNew(SHorizontalBox)
-						// Search
-						+ SHorizontalBox::Slot()
-						.VAlign(VAlign_Center)
-						.Padding(5, 0, 0, 0)
-						.FillWidth(1.0)
-						[
-							SAssignNew(SearchBoxPtr, SAssetSearchBox)
+	SearchBoxPtr = SNew(SAssetSearchBox)
 							.HintText(this, &SContentBrowser::GetSearchAssetsHintText)
 							.ShowSearchHistory(true)
 							.OnTextChanged(this, &SContentBrowser::OnSearchBoxChanged)
@@ -777,41 +884,16 @@ TSharedRef<SWidget> SContentBrowser::CreateAssetView(const FContentBrowserConfig
 							.OnAssetSearchBoxSuggestionChosen(this, &SContentBrowser::OnAssetSearchSuggestionChosen)
 							.DelayChangeNotificationsWhileTyping(true)
 							.Visibility((Config != nullptr ? Config->bCanShowAssetSearch : true) ? EVisibility::Visible : EVisibility::Collapsed)
-							.AddMetaData<FTagMetaData>(FTagMetaData(TEXT("ContentBrowserSearchAssets")))
-						]
-						
-						// Filter dropdown
-						+ SHorizontalBox::Slot()
-						.AutoWidth()
-						.HAlign(HAlign_Left)
-						.Padding(5, 0, 0, 0)
-						[
-							FilterComboButton.ToSharedRef()
-						]
-					]
-				]
-				+ SWrapBox::Slot()
-				.HAlign(HAlign_Left)
-				.VAlign(VAlign_Center)
-				.FillEmptySpace(true)
-				[
-					// Filter list
-					SNew(SBox)
-					.MinDesiredHeight_Lambda([this, SearchFilterItemHeight]() { return FilterListPtr->HasAnyFilters() ? SearchFilterItemHeight : 0; })
-					.VAlign(VAlign_Center)
-					[
-						FilterListPtr.ToSharedRef()
-					]
-				]
-			]
-		]
-		// Assets
-		+ SVerticalBox::Slot()
-		.FillHeight( 1.0f )
-		.Padding( 0, 0 )
-		[
-			AssetViewPtr.ToSharedRef()
-		];
+							.AddMetaData<FTagMetaData>(FTagMetaData(TEXT("ContentBrowserSearchAssets")));
+
+	/* Create the Border that the Asset View will live in, the actual layout is populated in OnFilterBarLayoutChanging, 
+	 * which is called initially called through SAssetFilterBar::LoadSettings through SContentBrowser::LoadSettings()
+	 */
+	AssetViewBorder = SNew(SBorder)
+			.BorderImage(FAppStyle::Get().GetBrush("Brushes.Panel"))
+			.Padding(FMargin(2.0f, 2.0f, 2.0f, 0.0f));
+
+	return AssetViewBorder.ToSharedRef();
 }
 
 TSharedRef<SWidget> SContentBrowser::CreateFavoritesView(const FContentBrowserConfig* Config)
diff --git a/Engine/Source/Editor/ContentBrowser/Private/SContentBrowser.h b/Engine/Source/Editor/ContentBrowser/Private/SContentBrowser.h
--- a/Engine/Source/Editor/ContentBrowser/Private/SContentBrowser.h
+++ b/Engine/Source/Editor/ContentBrowser/Private/SContentBrowser.h
@@ -34,6 +34,7 @@ class UFactory;
 class UToolMenu;
 class SSearchToggleButton;
 class UContentBrowserToolbarMenuContext;
+enum class EFilterBarLayout : uint8;
 
 struct FToolMenuContext;
 
@@ -556,6 +557,9 @@ private:
 	SSplitter::ESizeRule GetPathAreaSizeRule() const;
 	SSplitter::ESizeRule GetCollectionsAreaSizeRule() const;
 
+	/** Called when the layout of the SFilterList is changing */
+	void OnFilterBarLayoutChanging(EFilterBarLayout NewLayout);
+
 private:
 
 	/** The tab that contains this browser */
@@ -600,6 +604,12 @@ private:
 	/** The filter list */
 	TSharedPtr<SFilterList> FilterListPtr;
 
+	/** The Combo Button used to summon the filter dropdown */
+	TSharedPtr<SWidget> FilterComboButton;
+
+	/** The border that holds the content in AssetView */
+	TSharedPtr<SBorder> AssetViewBorder;
+
 	/** The path picker */
 	TSharedPtr<SComboButton> PathPickerButton;
 
diff --git a/Engine/Source/Editor/ContentBrowser/Private/SFilterList.cpp b/Engine/Source/Editor/ContentBrowser/Private/SFilterList.cpp
--- a/Engine/Source/Editor/ContentBrowser/Private/SFilterList.cpp
+++ b/Engine/Source/Editor/ContentBrowser/Private/SFilterList.cpp
@@ -46,6 +46,7 @@ void SFilterList::Construct( const FArguments& InArgs )
 {
 	bUseSharedSettings = InArgs._UseSharedSettings;
 	OnGetContextMenu = InArgs._OnGetContextMenu;
+	OnFilterBarLayoutChanging = InArgs._OnFilterBarLayoutChanging;
 	this->OnFilterChanged = InArgs._OnFilterChanged;
 	this->ActiveFilters = InArgs._FrontendFilters;
 	InitialClassFilters = InArgs._InitialClassFilters;
@@ -118,6 +119,8 @@ void SFilterList::Construct( const FArguments& InArgs )
 	Args._OnFilterChanged = this->OnFilterChanged;
 	Args._CreateTextFilter = InArgs._CreateTextFilter;
 	Args._FilterBarIdentifier = InArgs._FilterBarIdentifier;
+	Args._FilterBarLayout = InArgs._FilterBarLayout;
+	Args._CanChangeOrientation = InArgs._CanChangeOrientation;
 
 	SAssetFilterBar<FAssetFilterType>::Construct(Args);
 
@@ -303,6 +306,47 @@ bool IsFilteredByPicker(const TArray<UClass*>& FilterClassList, UClass* TestClas
 	return true;
 }
 
+void SFilterList::PopulateFilterDisplayMenu(UToolMenu* Menu)
+{
+	FToolMenuSection& Section = Menu->AddSection("FilterListFilterDisplay", LOCTEXT("FilterDisplay", "Filter Display"));
+
+	Section.AddMenuEntry(
+		"VerticalLayout",
+		LOCTEXT("FilterListVerticalLayout", "Vertical"),
+		LOCTEXT("FilterListVerticalLayoutToolTip", "Swap to a vertical layout for the filter bar"),
+		FSlateIcon(),
+		FUIAction(
+			FExecuteAction::CreateLambda([this]()
+			{
+				if(this->FilterBarLayout != EFilterBarLayout::Vertical)
+				{
+					this->SetFilterLayout(EFilterBarLayout::Vertical);
+				}
+			}),
+			FCanExecuteAction(),
+			FIsActionChecked::CreateLambda([this]() { return FilterBarLayout == EFilterBarLayout::Vertical; })),
+		EUserInterfaceActionType::RadioButton
+	);
+
+	Section.AddMenuEntry(
+		"HorizontalLayout",
+		LOCTEXT("FilterListHorizontalLayout", "Horizontal"),
+		LOCTEXT("FilterListHorizontalLayoutToolTip", "Swap to a Horizontal layout for the filter bar"),
+		FSlateIcon(),
+		FUIAction(
+		FExecuteAction::CreateLambda([this]()
+			{
+				if(this->FilterBarLayout != EFilterBarLayout::Horizontal)
+				{
+					this->SetFilterLayout(EFilterBarLayout::Horizontal);
+				}
+			}),
+			FCanExecuteAction(),
+			FIsActionChecked::CreateLambda([this]() { return FilterBarLayout == EFilterBarLayout::Horizontal; })),
+		EUserInterfaceActionType::RadioButton
+	);
+}
+
 void SFilterList::PopulateAddFilterMenu_Internal(UToolMenu* Menu)
 {
 	EAssetTypeCategories::Type MenuExpansion = EAssetTypeCategories::Basic;
@@ -315,6 +359,8 @@ void SFilterList::PopulateAddFilterMenu_Internal(UToolMenu* Menu)
 	{
 		return !IsFilteredByPicker(this->InitialClassFilters, TestClass);
 	}));
+	
+	PopulateFilterDisplayMenu(Menu);
 
 	Menu->AddSection("ContentBrowserFilterMiscAsset", LOCTEXT("MiscAssetsMenuHeading", "Misc Options") );
 }
@@ -579,9 +625,14 @@ void SFilterList::SaveSettings()
 			
 			InstanceSettings->CustomTextFilters.Add(InstanceFilterState);
 		}
-		
 	}
 
+	// Only save the orientation if we allow dynamic modification and saving
+	InstanceSettings->bIsLayoutSaved = this->bCanChangeOrientation;
+	if(this->bCanChangeOrientation)
+	{
+		InstanceSettings->FilterBarLayout = this->FilterBarLayout;
+	}
 
 	SaveConfig();
 }
@@ -616,6 +667,14 @@ void SFilterList::LoadSettings()
 	{
 		RestoreCustomTextFilterState(FilterState);
 	}
+
+	if(InstanceSettings->bIsLayoutSaved)
+	{
+		FilterBarLayout = InstanceSettings->FilterBarLayout;
+	}
+
+	// We want to call this even if the Layout isn't saved, to make sure OnFilterBarLayoutChanging is fired
+	SetFilterLayout(FilterBarLayout);
 	
 	this->OnFilterChanged.ExecuteIfBound();
 }
@@ -641,6 +700,53 @@ void SFilterList::LoadCustomTextFilters(const FFilterBarSettings* FilterBarConfi
 	}
 }
 
+void SFilterList::AddWidgetToCurrentLayout(TSharedRef<SWidget> InWidget)
+{
+	if(FilterBarLayout == EFilterBarLayout::Horizontal)
+	{
+		HorizontalFilterBox->AddSlot()
+		[
+			InWidget
+		];
+	}
+	else
+	{
+		VerticalFilterBox->AddSlot()
+		[
+			InWidget
+		];
+	}
+}
+
+void SFilterList::SetFilterLayout(EFilterBarLayout InFilterBarLayout)
+{
+	FilterBarLayout = InFilterBarLayout;
+ 		
+	if(FilterBarLayout == EFilterBarLayout::Horizontal)
+	{
+		VerticalFilterBox->ClearChildren();
+ 			
+		FilterBox->SetActiveWidget(HorizontalFilterBox.ToSharedRef());
+	}
+	else
+	{
+		HorizontalFilterBox->ClearChildren();
+        	
+		FilterBox->SetActiveWidget(VerticalFilterBox.ToSharedRef());
+	}
+
+	OnFilterBarLayoutChanging.ExecuteIfBound(FilterBarLayout);
+
+	for(TSharedRef<SFilter> Filter: Filters)
+	{
+		AddWidgetToLayout(Filter);
+	}
+ 		
+	this->Invalidate(EInvalidateWidgetReason::Layout);
+
+ 		
+}
+
 /////////////////////////////////////////
 // FFilterListCustomTextFilter
 /////////////////////////////////////////
@@ -695,5 +801,4 @@ TSharedPtr<FFilterBase<FAssetFilterType>> FFrontendFilter_CustomText::GetFilter(
 	return AsShared();
 }
 
-
 #undef LOCTEXT_NAMESPACE
diff --git a/Engine/Source/Editor/ContentBrowser/Private/SFilterList.h b/Engine/Source/Editor/ContentBrowser/Private/SFilterList.h
--- a/Engine/Source/Editor/ContentBrowser/Private/SFilterList.h
+++ b/Engine/Source/Editor/ContentBrowser/Private/SFilterList.h
@@ -29,6 +29,7 @@ class SFilterList : public SAssetFilterBar<FAssetFilterType>
 public:
 
 	DECLARE_DELEGATE_RetVal( TSharedPtr<SWidget>, FOnGetContextMenu );
+	DECLARE_DELEGATE_OneParam(FOnFilterBarLayoutChanging, EFilterBarLayout /* NewLayout */)
 	/**
 	 * An event delegate that is executed when a custom text filter has been created/modified/deleted in any FilterList
 	 * that is using shared settings.
@@ -37,8 +38,10 @@ public:
 	
 	using FOnFilterChanged = typename SAssetFilterBar<FAssetFilterType>::FOnFilterChanged;
 
-	SLATE_BEGIN_ARGS( SFilterList ):
-	_UseSharedSettings(false)
+	SLATE_BEGIN_ARGS( SFilterList )
+	: _UseSharedSettings(false)
+	, _FilterBarLayout(EFilterBarLayout::Horizontal)
+	, _CanChangeOrientation(false)
 	{
 		
 	}
@@ -71,6 +74,15 @@ public:
 		 */
 		SLATE_ARGUMENT(bool, UseSharedSettings)
 
+		/** Called when the filter bar layout is changing, before the filters are added to the layout */
+		SLATE_EVENT(FOnFilterBarLayoutChanging, OnFilterBarLayoutChanging)
+
+		/** The layout that determines how the filters are laid out */
+		SLATE_ARGUMENT(EFilterBarLayout, FilterBarLayout)
+			
+		/** If true, allow dynamically changing the orientation and saving in the config */
+		SLATE_ARGUMENT(bool, CanChangeOrientation)
+
 	SLATE_END_ARGS()
 
 	/** Constructs this widget with InArgs */
@@ -105,9 +117,14 @@ public:
 	/** Updates bIncludeClassName, bIncludeAssetPath and bIncludeCollectionNames for all custom text filters */
 	void UpdateCustomTextFilterIncludes(const bool InIncludeClassName, const bool InIncludeAssetPath, const bool InIncludeCollectionNames);
 
+	/** Add a custom widget to the filter bar alongside the filters */
+	void AddWidgetToCurrentLayout(TSharedRef<SWidget> InWidget);
+
 	virtual void SaveSettings()override;
 	virtual void LoadSettings() override;
 
+	virtual void SetFilterLayout(EFilterBarLayout InFilterBarLayout) override;
+
 protected:
 	
 	/** Handler for when the add filter button was clicked */
@@ -136,6 +153,9 @@ private:
 	/** Find the custom text filter corresponding to the specified state, and restore it's state to what is specified */
 	void RestoreCustomTextFilterState(const FCustomTextFilterState& InFilterState);
 	
+	/** Function to populate the options to change the filter layout */
+	void PopulateFilterDisplayMenu(UToolMenu* Menu);
+	
 private:
 
 	/** List of classes that our filters must match */
@@ -147,6 +167,9 @@ private:
 	/** Delegate for when filters have changed */
 	FOnFilterChanged OnFilterChanged;
 
+	/** Delegate for when the layout is changed */
+	FOnFilterBarLayoutChanging OnFilterBarLayoutChanging;
+
 	/** A reference to AllFrontEndFilters so we can access the filters as FFrontEndFilter instead of FFilterBase<FAssetFilterType> */
 	TArray< TSharedRef<FFrontendFilter> > AllFrontendFilters_Internal;
 
diff --git a/Engine/Source/Editor/EditorWidgets/Public/Filters/FilterBarConfig.h b/Engine/Source/Editor/EditorWidgets/Public/Filters/FilterBarConfig.h
--- a/Engine/Source/Editor/EditorWidgets/Public/Filters/FilterBarConfig.h
+++ b/Engine/Source/Editor/EditorWidgets/Public/Filters/FilterBarConfig.h
@@ -5,6 +5,7 @@
 #include "Containers/Map.h"
 #include "EditorConfigBase.h"
 #include "Filters/SCustomTextFilterDialog.h"
+#include "Filters/SBasicFilterBar.h"
 
 #include "FilterBarConfig.generated.h"
 
@@ -26,7 +27,6 @@ public:
 	/* The data inside the custom text filter */
 	UPROPERTY()
 	FCustomTextFilterData FilterData;
-	
 };
 
 USTRUCT()
@@ -48,6 +48,12 @@ public:
 	UPROPERTY()
 	TArray<FCustomTextFilterState> CustomTextFilters;
 
+	UPROPERTY()
+	bool bIsLayoutSaved = false;
+
+	UPROPERTY()
+	EFilterBarLayout FilterBarLayout = EFilterBarLayout::Horizontal;
+	
 	void Empty()
 	{
 		CustomFilters.Empty();
diff --git a/Engine/Source/Editor/EditorWidgets/Public/Filters/SAssetFilterBar.h b/Engine/Source/Editor/EditorWidgets/Public/Filters/SAssetFilterBar.h
--- a/Engine/Source/Editor/EditorWidgets/Public/Filters/SAssetFilterBar.h
+++ b/Engine/Source/Editor/EditorWidgets/Public/Filters/SAssetFilterBar.h
@@ -108,6 +108,8 @@ public:
 	
 	SLATE_BEGIN_ARGS( SAssetFilterBar<FilterType> )
 		: _UseDefaultAssetFilters(true)
+		, _FilterBarLayout(EFilterBarLayout::Horizontal)
+		, _CanChangeOrientation(false)
 		{
 		
 		}
@@ -140,6 +142,12 @@ public:
 	
 		/** Whether the filter bar should provide the default asset filters */
 		SLATE_ARGUMENT(bool, UseDefaultAssetFilters)
+		
+		/** The layout that determines how the filters are laid out */
+		SLATE_ARGUMENT(EFilterBarLayout, FilterBarLayout)
+		
+		/** If true, allow dynamically changing the orientation and saving in the config */
+		SLATE_ARGUMENT(bool, CanChangeOrientation)
 	
 	SLATE_END_ARGS()
 
@@ -157,6 +165,8 @@ public:
 		Args._OnExtendAddFilterMenu = InArgs._OnExtendAddFilterMenu;
 		Args._CreateTextFilter = InArgs._CreateTextFilter;
 		Args._FilterSearchBox = InArgs._FilterSearchBox;
+		Args._FilterBarLayout = InArgs._FilterBarLayout;
+		Args._CanChangeOrientation = InArgs._CanChangeOrientation;
 		
 		SBasicFilterBar<FilterType>::Construct(Args);
 
@@ -189,6 +199,13 @@ public:
 		SaveCustomTextFilters(FilterBarConfig);
 		SaveFilters(FilterBarConfig);
 
+		// Only save the orientation if we allow dynamic modification and saving
+		FilterBarConfig->bIsLayoutSaved = this->bCanChangeOrientation;
+		if(this->bCanChangeOrientation)
+		{
+			FilterBarConfig->FilterBarLayout = this->FilterBarLayout;
+		}
+
 		SaveConfig();
 	}
 
@@ -211,9 +228,15 @@ public:
 		LoadFilters(FilterBarConfig);
 		LoadCustomTextFilters(FilterBarConfig);
 
+		/* Only load the setting if we saved it, aka if SaveSettings() was ever called and bCanChangeOrientation is true */
+		if(FilterBarConfig->bIsLayoutSaved)
+		{
+			this->SetFilterLayout(FilterBarConfig->FilterBarLayout);
+		}
+
 		this->OnFilterChanged.ExecuteIfBound();
 	}
-	
+
 protected:
 
 	/** Save all the custom class filters (created by the user) into the specified config */
diff --git a/Engine/Source/Editor/EditorWidgets/Public/Filters/SFilterBar.h b/Engine/Source/Editor/EditorWidgets/Public/Filters/SFilterBar.h
--- a/Engine/Source/Editor/EditorWidgets/Public/Filters/SFilterBar.h
+++ b/Engine/Source/Editor/EditorWidgets/Public/Filters/SFilterBar.h
@@ -47,6 +47,8 @@ public:
 	SLATE_BEGIN_ARGS( SFilterBar<FilterType> )
 		: _FilterBarIdentifier(NAME_None)
 		, _UseDefaultAssetFilters(true)
+	    , _FilterBarLayout(EFilterBarLayout::Horizontal)
+		, _CanChangeOrientation(false)
 		{
 		
 		}
@@ -85,6 +87,12 @@ public:
 		 *	NOTE: Will bind a delegate to SFilterSearchBox::OnClickedAddSearchHistoryButton
 		 */
 		SLATE_ARGUMENT(TSharedPtr<SFilterSearchBox>, FilterSearchBox)
+
+		/** The layout that determines how the filters are laid out */
+		SLATE_ARGUMENT(EFilterBarLayout, FilterBarLayout)
+			
+		/** If true, allow dynamically changing the orientation and saving in the config */
+		SLATE_ARGUMENT(bool, CanChangeOrientation)
 	
 	SLATE_END_ARGS()
 
@@ -100,6 +108,8 @@ public:
 		Args._CreateTextFilter = InArgs._CreateTextFilter;
 		Args._FilterSearchBox = InArgs._FilterSearchBox;
 		Args._FilterBarIdentifier = InArgs._FilterBarIdentifier;
+		Args._FilterBarLayout = InArgs._FilterBarLayout;
+        Args._CanChangeOrientation = InArgs._CanChangeOrientation;
 		
 		SAssetFilterBar<FilterType>::Construct(Args);
 		
diff --git a/Engine/Source/Runtime/SlateCore/Private/Styling/StarshipCoreStyle.cpp b/Engine/Source/Runtime/SlateCore/Private/Styling/StarshipCoreStyle.cpp
--- a/Engine/Source/Runtime/SlateCore/Private/Styling/StarshipCoreStyle.cpp
+++ b/Engine/Source/Runtime/SlateCore/Private/Styling/StarshipCoreStyle.cpp
@@ -405,6 +405,8 @@ TSharedRef<ISlateStyle> FStarshipCoreStyle::Create()
 		Style->Set("Icons.FlipVertical", new IMAGE_BRUSH_SVG("Starship/Common/FlipVertical", Icon16x16));
 		Style->Set("Icons.Layout", new IMAGE_BRUSH_SVG("Starship/Common/Layout", Icon16x16));
 
+		Style->Set("Icons.BadgeModified", new IMAGE_BRUSH_SVG("Starship/Common/badge-modified", Icon16x16));
+		
 		// Toolbar Size Icons
 		Style->Set("Icons.Toolbar.Settings", new IMAGE_BRUSH_SVG("Starship/Common/settings", Icon20x20));
 	}
