SAssetSearchBox: Forward Focus to the internal text box properly. This fixes a bug where asset pickers in details views etc would search in the menu instead of the asset picker

diff --git a/Engine/Source/Developer/ToolWidgets/Private/Filters/SFilterSearchBox.cpp b/Engine/Source/Developer/ToolWidgets/Private/Filters/SFilterSearchBox.cpp
--- a/Engine/Source/Developer/ToolWidgets/Private/Filters/SFilterSearchBox.cpp
+++ b/Engine/Source/Developer/ToolWidgets/Private/Filters/SFilterSearchBox.cpp
@@ -110,12 +110,23 @@ void SFilterSearchBox::Construct( const FArguments& InArgs )
 	];
 }
 
+bool SFilterSearchBox::SupportsKeyboardFocus() const
+{
+	return SearchBox->SupportsKeyboardFocus();
+}
+
 bool SFilterSearchBox::HasKeyboardFocus() const
 {
 	// Since keyboard focus is forwarded to our editable text, we will test it instead
 	return SearchBox->HasKeyboardFocus();
 }
 
+FReply SFilterSearchBox::OnFocusReceived( const FGeometry& MyGeometry, const FFocusEvent& InFocusEvent )
+{
+	// Forward keyboard focus to our editable text widget
+	return FReply::Handled().SetUserFocus(SearchBox.ToSharedRef(), InFocusEvent.GetCause());
+}
+
 /** Handler for when text in the editable text box changed */
 void SFilterSearchBox::HandleTextChanged(const FText& NewText)
 {
diff --git a/Engine/Source/Developer/ToolWidgets/Public/Filters/SFilterSearchBox.h b/Engine/Source/Developer/ToolWidgets/Public/Filters/SFilterSearchBox.h
--- a/Engine/Source/Developer/ToolWidgets/Public/Filters/SFilterSearchBox.h
+++ b/Engine/Source/Developer/ToolWidgets/Public/Filters/SFilterSearchBox.h
@@ -75,6 +75,8 @@ public:
 	/** SWidget interface */
 	void OnFocusLost(const FFocusEvent& InFocusEvent) override;
 	virtual bool HasKeyboardFocus() const override;
+	bool SupportsKeyboardFocus() const;
+	FReply OnFocusReceived( const FGeometry& MyGeometry, const FFocusEvent& InFocusEvent );
 
 	/** Show a + button next to the current search, and set the handler for when that is clicked */
 	void SetOnSaveSearchHandler(FOnSaveSearchClicked InOnSaveSearchHandler);
diff --git a/Engine/Source/Editor/EditorWidgets/Private/SAssetSearchBox.cpp b/Engine/Source/Editor/EditorWidgets/Private/SAssetSearchBox.cpp
--- a/Engine/Source/Editor/EditorWidgets/Private/SAssetSearchBox.cpp
+++ b/Engine/Source/Editor/EditorWidgets/Private/SAssetSearchBox.cpp
@@ -188,7 +188,7 @@ bool SAssetSearchBox::HasKeyboardFocus() const
 FReply SAssetSearchBox::OnFocusReceived( const FGeometry& MyGeometry, const FFocusEvent& InFocusEvent )
 {
 	// Forward keyboard focus to our editable text widget
-	return FReply::Handled().SetUserFocus(InputText.ToSharedRef(), InFocusEvent.GetCause());
+	return InputText->OnFocusReceived(MyGeometry, InFocusEvent);
 }
 
 void SAssetSearchBox::HandleTextChanged(const FText& NewText)
