Fix SCustomDialog to give room to buttons when content is narrow.

diff --git a/Engine/Source/Developer/ToolWidgets/Private/Dialog/SCustomDialog.cpp b/Engine/Source/Developer/ToolWidgets/Private/Dialog/SCustomDialog.cpp
--- a/Engine/Source/Developer/ToolWidgets/Private/Dialog/SCustomDialog.cpp
+++ b/Engine/Source/Developer/ToolWidgets/Private/Dialog/SCustomDialog.cpp
@@ -49,6 +49,7 @@ void SCustomDialog::Construct(const FArguments& InArgs)
 				SNew(SVerticalBox)
 				+ SVerticalBox::Slot()
 				.FillHeight(1.0f)
+				.HAlign(HAlign_Right)
 				[
 					CreateContentBox(InArgs)
 				]
@@ -138,7 +139,7 @@ TSharedRef<SWidget> SCustomDialog::CreateButtonBox(const FArguments& InArgs)
 	
 		// Before buttons
 		+SHorizontalBox::Slot()
-			.FillWidth(1.f)
+			.AutoWidth()
 			.HAlign(HAlign_Left)
 			.VAlign(VAlign_Center)
 			[
