Implemented GetValues(), added a function to query the number of checkboxes in the list, and added a variable to turn off the display of the global checkbox in the header column

diff --git a/Engine/Source/Developer/ToolWidgets/Private/SCheckBoxList.cpp b/Engine/Source/Developer/ToolWidgets/Private/SCheckBoxList.cpp
--- a/Engine/Source/Developer/ToolWidgets/Private/SCheckBoxList.cpp
+++ b/Engine/Source/Developer/ToolWidgets/Private/SCheckBoxList.cpp
@@ -101,6 +101,8 @@ void SCheckBoxList::Construct(const FArguments& InArgs, const TArray<TSharedRef<
 		Items.Add(MakeShared<CheckBoxList::FItemPair>(Widget, bIsChecked));
 	}
 
+	bool bShowHeaderCheckbox = InArgs._IncludeGlobalCheckBoxInHeaderRow;
+
 	TSharedRef<SHeaderRow> HeaderRowWidget = SNew(SHeaderRow)
 		+ SHeaderRow::Column(CheckBoxList::ColumnID_CheckBox)
 		[
@@ -108,12 +110,13 @@ void SCheckBoxList::Construct(const FArguments& InArgs, const TArray<TSharedRef<
 			.Style(InArgs._CheckBoxStyle)
 			.IsChecked(this, &SCheckBoxList::GetToggleSelectedState)
 			.OnCheckStateChanged(this, &SCheckBoxList::OnToggleSelectedCheckBox)
+			.Visibility_Lambda([bShowHeaderCheckbox] { return bShowHeaderCheckbox ? EVisibility::Visible : EVisibility::Hidden; })
 		]
 		.FixedWidth(CheckBoxList::CheckBoxColumnWidth)
 		+ SHeaderRow::Column(CheckBoxList::ColumnID_Item)
 		.DefaultLabel(InArgs._ItemHeaderLabel)
 		.FillWidth(1.0f);
-	
+
 	ChildSlot
 	[
 		SAssignNew(ListView, SListView<TSharedRef<CheckBoxList::FItemPair>>)
@@ -153,6 +156,21 @@ bool SCheckBoxList::IsItemChecked(int32 Index) const
 	return Items.IsValidIndex(Index) ? Items[Index]->bIsChecked : false;
 }
 
+TArray<bool> SCheckBoxList::GetValues() const
+{
+	TArray<bool> Values;
+	for (const TSharedRef<CheckBoxList::FItemPair>& Item : Items)
+	{
+		Values.Add(Item->bIsChecked);
+	}
+	return Values;
+}
+
+int32 SCheckBoxList::GetNumCheckboxes() const
+{
+	return Items.Num();
+}
+
 void SCheckBoxList::UpdateAllChecked()
 {
 	bool bContainsTrue = Items.ContainsByPredicate([](const TSharedRef<CheckBoxList::FItemPair>& Item) { return Item->bIsChecked; });
diff --git a/Engine/Source/Developer/ToolWidgets/Public/SCheckBoxList.h b/Engine/Source/Developer/ToolWidgets/Public/SCheckBoxList.h
--- a/Engine/Source/Developer/ToolWidgets/Public/SCheckBoxList.h
+++ b/Engine/Source/Developer/ToolWidgets/Public/SCheckBoxList.h
@@ -27,11 +27,14 @@ class TOOLWIDGETS_API SCheckBoxList: public SCompoundWidget
 {
 	SLATE_BEGIN_ARGS(SCheckBoxList)
 		: _CheckBoxStyle(&FCoreStyle::Get().GetWidgetStyle<FCheckBoxStyle>("Checkbox"))
+		, _IncludeGlobalCheckBoxInHeaderRow(true)
 		{}
 		/** The styling of the CheckBox */
 		SLATE_STYLE_ARGUMENT(FCheckBoxStyle, CheckBoxStyle)
 		/** The label of the item column header */
 		SLATE_ARGUMENT(FText, ItemHeaderLabel)
+		/** The label of the item column header */
+		SLATE_ARGUMENT(bool, IncludeGlobalCheckBoxInHeaderRow)
 	SLATE_END_ARGS()
 
 public:
@@ -45,6 +48,7 @@ public:
 	bool IsItemChecked(int32 Index) const;
 
 	TArray<bool> GetValues() const;
+	int32 GetNumCheckboxes() const;
 
 private:
 	void UpdateAllChecked();
