Filter Widgets: Fix an issue where the search text could go over the save search button in the Content Browser
Fix an issue where the text boxes in the custom text filter dialog would resize

diff --git a/Engine/Source/Developer/ToolWidgets/Private/Filters/SCustomTextFilterDialog.cpp b/Engine/Source/Developer/ToolWidgets/Private/Filters/SCustomTextFilterDialog.cpp
--- a/Engine/Source/Developer/ToolWidgets/Private/Filters/SCustomTextFilterDialog.cpp
+++ b/Engine/Source/Developer/ToolWidgets/Private/Filters/SCustomTextFilterDialog.cpp
@@ -5,6 +5,7 @@
 #include "Widgets/Colors/SColorBlock.h"
 #include "Widgets/Colors/SColorPicker.h"
 #include "Widgets/Text/STextBlock.h"
+#include "Widgets/Layout/SBox.h"
 #include "SPrimaryButton.h"
 
 #define LOCTEXT_NAMESPACE "CustomTextFilter"
@@ -112,16 +113,20 @@ void SCustomTextFilterDialog::Construct( const FArguments& InArgs )
 				.AutoWidth()
 				.Padding(21, 0, 0, 0)
 				[
-					SAssignNew(FilterLabelTextBox, SEditableTextBox)
-					.MinDesiredWidth(254)
-					.Text_Lambda([this] ()
-					{
-						return this->FilterData.FilterLabel;
-					})
-					.OnTextChanged_Lambda([this] (const FText& InText)
-					{
-						this->FilterData.FilterLabel = InText;
-					})
+					SNew(SBox)
+					.WidthOverride(254)
+					[
+						SAssignNew(FilterLabelTextBox, SEditableTextBox)
+						.Text_Lambda([this] ()
+						{
+							return this->FilterData.FilterLabel;
+						})
+						.OnTextChanged_Lambda([this] (const FText& InText)
+						{
+							this->FilterData.FilterLabel = InText;
+						})
+					]
+					
 				]
 
 				// Filter Color Text
@@ -170,16 +175,20 @@ void SCustomTextFilterDialog::Construct( const FArguments& InArgs )
 				.AutoWidth()
 				.Padding(21, 0, 0, 0)
 				[
-					SNew(SEditableTextBox)
-					.MinDesiredWidth(556)
-					.Text_Lambda([this] ()
-					{
-						return this->FilterData.FilterString;
-					})
-					.OnTextChanged_Lambda([this] (const FText& InText)
-					{
-						this->FilterData.FilterString = InText;
-					})
+					SNew(SBox)
+					.WidthOverride(556)
+					[
+						SNew(SEditableTextBox)
+						.Text_Lambda([this] ()
+						{
+							return this->FilterData.FilterString;
+						})
+						.OnTextChanged_Lambda([this] (const FText& InText)
+						{
+							this->FilterData.FilterString = InText;
+						})
+					]
+					
 				]
 			]
 			// Buttons
diff --git a/Engine/Source/Developer/ToolWidgets/Private/Filters/SFilterSearchBox.cpp b/Engine/Source/Developer/ToolWidgets/Private/Filters/SFilterSearchBox.cpp
--- a/Engine/Source/Developer/ToolWidgets/Private/Filters/SFilterSearchBox.cpp
+++ b/Engine/Source/Developer/ToolWidgets/Private/Filters/SFilterSearchBox.cpp
@@ -33,6 +33,7 @@ void SFilterSearchBox::Construct( const FArguments& InArgs )
 			+ SOverlay::Slot()
 			[
 				SAssignNew(SearchBox, SSearchBox)
+				.Style(&FAppStyle::GetWidgetStyle<FSearchBoxStyle>("FilterBar.SearchBox"))
 				.InitialText(InArgs._InitialText)
 				.HintText(InArgs._HintText)
 				.OnTextChanged(this, &SFilterSearchBox::HandleTextChanged)
diff --git a/Engine/Source/Developer/ToolWidgets/Public/Filters/SBasicFilterBar.h b/Engine/Source/Developer/ToolWidgets/Public/Filters/SBasicFilterBar.h
--- a/Engine/Source/Developer/ToolWidgets/Public/Filters/SBasicFilterBar.h
+++ b/Engine/Source/Developer/ToolWidgets/Public/Filters/SBasicFilterBar.h
@@ -460,6 +460,7 @@ protected:
 					.Padding(0.0f)
 					.IsChecked(this, &SFilter::IsChecked)
 					.OnCheckStateChanged(this, &SFilter::FilterToggled)
+					.CheckBoxContentUsesAutoWidth(false)
 					.OnGetMenuContent(this, &SFilter::GetRightClickMenuContent)
 					[
 						SNew(SHorizontalBox)
diff --git a/Engine/Source/Runtime/SlateCore/Private/Styling/StarshipCoreStyle.cpp b/Engine/Source/Runtime/SlateCore/Private/Styling/StarshipCoreStyle.cpp
--- a/Engine/Source/Runtime/SlateCore/Private/Styling/StarshipCoreStyle.cpp
+++ b/Engine/Source/Runtime/SlateCore/Private/Styling/StarshipCoreStyle.cpp
@@ -842,6 +842,18 @@ TSharedRef<ISlateStyle> FStarshipCoreStyle::Create()
 		);
 	}
 
+	// Filter Widgets
+	{
+		FSearchBoxStyle FilterSearchBoxStyle = Style->GetWidgetStyle<FSearchBoxStyle>("SearchBox");
+		FEditableTextBoxStyle FilterTextBoxStyle = FilterSearchBoxStyle.TextBoxStyle;
+
+		// Extra padding on the right to account for the Save Search button
+		FilterTextBoxStyle.SetPadding(FMargin(8.0f, 3.0f, 30.0f, 4.0f));
+		FilterSearchBoxStyle.SetTextBoxStyle(FilterTextBoxStyle);
+
+		Style->Set("FilterBar.SearchBox", FilterSearchBoxStyle);
+	}
+	
 	// Project Launcher
 	{
 		Style->Set("Launcher.TabIcon", new IMAGE_BRUSH_SVG("Starship/Common/ProjectLauncher", Icon16x16));
