SFilterSearchBox: Fix an issue where the error icon would overlap the search history dropdown

diff --git a/Engine/Source/Developer/ToolWidgets/Private/Filters/SFilterSearchBox.cpp b/Engine/Source/Developer/ToolWidgets/Private/Filters/SFilterSearchBox.cpp
--- a/Engine/Source/Developer/ToolWidgets/Private/Filters/SFilterSearchBox.cpp
+++ b/Engine/Source/Developer/ToolWidgets/Private/Filters/SFilterSearchBox.cpp
@@ -11,13 +11,148 @@
 
 #define LOCTEXT_NAMESPACE "CustomTextFilter"
 
+/** Internal subclass of SSearchBox, to allow adding the save search and search history buttons to the
+ *  Horizontal Box internal to SEditableTextBox
+ */
+class SFilterSearchBoxImpl : public SSearchBox
+{
+public:
+	
+	DECLARE_DELEGATE_RetVal( EVisibility, FGetSearchHistoryVisibility )
+
+	SLATE_BEGIN_ARGS(SFilterSearchBoxImpl)
+		: _HintText( LOCTEXT("SearchHint", "Search") )
+		, _InitialText()
+		, _OnTextChanged()
+		, _OnTextCommitted()
+		, _DelayChangeNotificationsWhileTyping( true )
+	{ }
+
+		/** The text displayed in the SearchBox when no text has been entered */
+		SLATE_ATTRIBUTE( FText, HintText )
+
+		/** The text displayed in the SearchBox when it's created */
+		SLATE_ATTRIBUTE( FText, InitialText )
+	
+		/** Invoked whenever the text changes */
+		SLATE_EVENT( FOnTextChanged, OnTextChanged )
+
+		/** Invoked whenever the text is committed (e.g. user presses enter) */
+		SLATE_EVENT( FOnTextCommitted, OnTextCommitted )
+	
+		/** Whether the SearchBox should delay notifying listeners of text changed events until the user is done typing */
+		SLATE_ATTRIBUTE( bool, DelayChangeNotificationsWhileTyping )
+
+		/** Callback delegate to have first chance handling of the OnKeyDown event */
+		SLATE_EVENT(FOnKeyDown, OnKeyDownHandler)
+
+		/** Delegate for when the Plus icon next to an item in the search history is clicked */
+		SLATE_EVENT(SFilterSearchBox::FOnSaveSearchClicked, OnSaveSearchClicked)
+	
+		/** Delegate for when the search history chevron is clicked */
+		SLATE_EVENT(FOnClicked, OnSearchHistoryClicked)
+
+		/** Delegate to get the visibility of the search history chevron */
+		SLATE_EVENT(FGetSearchHistoryVisibility, GetSearchHistoryVisibility)
+
+	SLATE_END_ARGS()
+	
+	void Construct( const FArguments& InArgs );
+	
+	/** Show a + button next to the current search, and set the handler for when that is clicked */
+	void SetOnSaveSearchHandler(SFilterSearchBox::FOnSaveSearchClicked InOnSaveSearchHandler);
+
+protected:
+
+	/** Delegate for when the Plus icon next to an item in the search history is clicked */
+	SFilterSearchBox::FOnSaveSearchClicked OnSaveSearchClicked;
+	
+	/** Delegate to get the visibility of the search history chevron */
+	FGetSearchHistoryVisibility GetSearchHistoryVisibility;
+};
+
+void SFilterSearchBoxImpl::Construct(const FArguments& InArgs)
+{
+	SSearchBox::Construct( SSearchBox::FArguments()
+		.Style(&FAppStyle::GetWidgetStyle<FSearchBoxStyle>("FilterBar.SearchBox"))
+		.InitialText(InArgs._InitialText)
+		.HintText(InArgs._HintText)
+		.OnTextChanged(InArgs._OnTextChanged)
+		.OnTextCommitted(InArgs._OnTextCommitted)
+		.SelectAllTextWhenFocused( false )
+		.DelayChangeNotificationsWhileTyping( InArgs._DelayChangeNotificationsWhileTyping )
+		.OnKeyDownHandler(InArgs._OnKeyDownHandler)
+	);
+
+	OnSaveSearchClicked = InArgs._OnSaveSearchClicked;
+	GetSearchHistoryVisibility = InArgs._GetSearchHistoryVisibility;
+
+	// + Button to save search
+	Box->AddSlot()
+		.VAlign(VAlign_Center)
+		.HAlign(HAlign_Right)
+		.AutoWidth()
+		.Padding(-4, 0, 0, 0) // Negative padding intentional to have the same padding between all buttons
+		[
+			// Button to save the currently occuring search 
+			SNew(SButton)
+			.ContentPadding(0)
+			.Visibility_Lambda([this]()
+			{
+				// Only visible if there is a search active currently and the OnSaveSearchClicked delegate is bound
+				return this->GetText().IsEmpty() || !this->OnSaveSearchClicked.IsBound() ? EVisibility::Collapsed : EVisibility::Visible;
+			})
+			.ButtonStyle(FAppStyle::Get(), "HoverOnlyButton")
+			.OnClicked_Lambda([this]()
+			{
+				this->OnSaveSearchClicked.ExecuteIfBound(this->GetText());
+				return FReply::Handled();
+			})
+			[
+				SNew(SImage)
+				.Image(FAppStyle::Get().GetBrush("Icons.PlusCircle"))
+				.ColorAndOpacity(FSlateColor::UseForeground())
+			]
+		];
+
+	// Chevron to show search history
+	Box->AddSlot()
+		.HAlign(HAlign_Right)
+		.VAlign(VAlign_Center)
+		.AutoWidth()
+		.Padding(0, 0, 4, 0)
+		[
+			// Chevron to open the search history dropdown
+			SNew(SButton)
+			.ContentPadding(0)
+			.ButtonStyle(FAppStyle::Get(), "HoverOnlyButton")
+			.ClickMethod(EButtonClickMethod::MouseDown)
+			.OnClicked(InArgs._OnSearchHistoryClicked)
+			.ToolTipText(LOCTEXT("SearchHistoryToolTipText", "Click to show the Search History"))
+			.Visibility_Lambda([this]()
+			{
+				return this->GetSearchHistoryVisibility.Execute();
+			})
+			[
+				SNew(SImage)
+				.Image(FAppStyle::Get().GetBrush("Icons.ChevronDown"))
+				.ColorAndOpacity(FSlateColor::UseForeground())
+			]
+		];
+	
+}
+
+void SFilterSearchBoxImpl::SetOnSaveSearchHandler(SFilterSearchBox::FOnSaveSearchClicked InOnSaveSearchHandler)
+{
+	OnSaveSearchClicked = InOnSaveSearchHandler;
+}
+
 void SFilterSearchBox::Construct( const FArguments& InArgs )
 {
 	MaxSearchHistory = InArgs._MaxSearchHistory;
 	OnTextChanged = InArgs._OnTextChanged;
 	OnTextCommitted = InArgs._OnTextCommitted;
 	bShowSearchHistory = InArgs._ShowSearchHistory;
-	OnSaveSearchClicked = InArgs._OnSaveSearchClicked;
 
 	// Default text shown when there are no items in the search history
 	EmptySearchHistoryText = MakeShareable(new FText(LOCTEXT("EmptySearchHistoryText", "The Search History is Empty")));
@@ -28,70 +163,16 @@ void SFilterSearchBox::Construct( const FArguments& InArgs )
 		SAssignNew(SearchHistoryBox, SMenuAnchor)
 		.Placement(EMenuPlacement::MenuPlacement_ComboBoxRight)
 		[
-			SNew(SOverlay)
-			
-			+ SOverlay::Slot()
-			[
-				SAssignNew(SearchBox, SSearchBox)
-				.Style(&FAppStyle::GetWidgetStyle<FSearchBoxStyle>("FilterBar.SearchBox"))
-				.InitialText(InArgs._InitialText)
-				.HintText(InArgs._HintText)
-				.OnTextChanged(this, &SFilterSearchBox::HandleTextChanged)
-				.OnTextCommitted(this, &SFilterSearchBox::HandleTextCommitted)
-				.SelectAllTextWhenFocused( false )
-				.DelayChangeNotificationsWhileTyping( InArgs._DelayChangeNotificationsWhileTyping )
-				.OnKeyDownHandler(InArgs._OnKeyDownHandler)
-			]
-			+ SOverlay::Slot()
-			.HAlign(HAlign_Right)
-			[
-				SNew(SHorizontalBox)
-
-				+SHorizontalBox::Slot()
-				.VAlign(VAlign_Center)
-				.HAlign(HAlign_Right)
-				[
-					// Button to save the currently occuring search 
-					SNew(SButton)
-					.ContentPadding(0)
-					.Visibility_Lambda([this]()
-					{
-						// Only visible if there is a search active currently and the OnSaveSearchClicked delegate is bound
-						return this->GetText().IsEmpty() || !this->OnSaveSearchClicked.IsBound() ? EVisibility::Collapsed : EVisibility::Visible;
-					})
-					.ButtonStyle(FAppStyle::Get(), "HoverOnlyButton")
-					.OnClicked_Lambda([this]()
-					{
-						this->OnSaveSearchClicked.ExecuteIfBound(this->GetText());
-						return FReply::Handled();
-					})
-					[
-						SNew(SImage)
-						.Image(FAppStyle::Get().GetBrush("Icons.PlusCircle"))
-						.ColorAndOpacity(FSlateColor::UseForeground())
-					]
-				]
-
-				+SHorizontalBox::Slot()
-				.HAlign(HAlign_Right)
-				.VAlign(VAlign_Center)
-				.Padding(0, 0, 4, 0)
-				[
-					// Chevron to open the search history dropdown
-					SNew(SButton)
-					.ContentPadding(0)
-					.ButtonStyle(FAppStyle::Get(), "HoverOnlyButton")
-					.ClickMethod(EButtonClickMethod::MouseDown)
-					.OnClicked(this, &SFilterSearchBox::OnClickedSearchHistory)
-					.ToolTipText(LOCTEXT("SearchHistoryToolTipText", "Click to show the Search History"))
-					.Visibility(this, &SFilterSearchBox::GetSearchHistoryVisibility)
-					[
-						SNew(SImage)
-						.Image(FAppStyle::Get().GetBrush("Icons.ChevronDown"))
-						.ColorAndOpacity(FSlateColor::UseForeground())
-					]
-				]
-			]
+			SAssignNew(SearchBox, SFilterSearchBoxImpl)
+			.InitialText(InArgs._InitialText)
+			.HintText(InArgs._HintText)
+			.OnTextChanged(this, &SFilterSearchBox::HandleTextChanged)
+			.OnTextCommitted(this, &SFilterSearchBox::HandleTextCommitted)
+			.DelayChangeNotificationsWhileTyping( InArgs._DelayChangeNotificationsWhileTyping )
+			.OnKeyDownHandler(InArgs._OnKeyDownHandler)
+			.OnSearchHistoryClicked(this, &SFilterSearchBox::OnClickedSearchHistory)
+			.GetSearchHistoryVisibility(this, &SFilterSearchBox::GetSearchHistoryVisibility)
+			.OnSaveSearchClicked(InArgs._OnSaveSearchClicked)
 		]
 		.MenuContent
 		(
@@ -260,7 +341,7 @@ FReply SFilterSearchBox::OnClickedSearchHistory()
 
 void SFilterSearchBox::SetOnSaveSearchHandler(FOnSaveSearchClicked InOnSaveSearchHandler)
 {
-	OnSaveSearchClicked = InOnSaveSearchHandler;
+	SearchBox->SetOnSaveSearchHandler(InOnSaveSearchHandler);
 }
 
 EVisibility SFilterSearchBox::GetSearchHistoryVisibility() const
diff --git a/Engine/Source/Developer/ToolWidgets/Public/Filters/SFilterSearchBox.h b/Engine/Source/Developer/ToolWidgets/Public/Filters/SFilterSearchBox.h
--- a/Engine/Source/Developer/ToolWidgets/Public/Filters/SFilterSearchBox.h
+++ b/Engine/Source/Developer/ToolWidgets/Public/Filters/SFilterSearchBox.h
@@ -9,6 +9,7 @@
 class SMenuAnchor;
 class SSearchBox;
 template <typename T> class SListView;
+class SFilterSearchBoxImpl;
 
 /** A SearchBox widget that contains support for displaying search history
  * When used with a FilterBar widget (@see SBasicFilterBar or SFilterBar), it allows the user to save searches from the
@@ -110,7 +111,7 @@ private:
 	TAttribute<bool> bShowSearchHistory;
 	
 	/** The actual Search Box */
-	TSharedPtr<SSearchBox> SearchBox;
+	TSharedPtr<SFilterSearchBoxImpl> SearchBox;
 	
 	/** Delegate for when text is changed in the edit box */
 	FOnTextChanged OnTextChanged;
@@ -118,9 +119,6 @@ private:
 	/** Delegate for when text is changed in the edit box */
 	FOnTextCommitted OnTextCommitted;
 
-	/** Delegate for when the Plus icon next to an item in the search history is clicked */
-	FOnSaveSearchClicked OnSaveSearchClicked;
-
 	/** The max amount of items to show in the Search History */
 	int32 MaxSearchHistory;
 
