SCustomDialog no longer clips any button texts that are longer than the first button text.

This is because the dialog used SUniformGridPanel, which does sizing based on the first button, not the largest one. We now use SHorizontalBox with AutoWidth() instead.

Removed useless log entry.

#rnx

diff --git a/Engine/Source/Developer/ToolWidgets/Private/Dialog/SCustomDialog.cpp b/Engine/Source/Developer/ToolWidgets/Private/Dialog/SCustomDialog.cpp
--- a/Engine/Source/Developer/ToolWidgets/Private/Dialog/SCustomDialog.cpp
+++ b/Engine/Source/Developer/ToolWidgets/Private/Dialog/SCustomDialog.cpp
@@ -4,7 +4,6 @@
 
 #include "Framework/Application/SlateApplication.h"
 #include "Framework/Docking/TabManager.h"
-#include "Logging/LogMacros.h"
 #include "Styling/AppStyle.h"
 #include "Styling/SlateBrush.h"
 #include "Widgets/Images/SImage.h"
@@ -12,12 +11,9 @@
 #include "Widgets/Text/STextBlock.h"
 #include "Widgets/Layout/SBox.h"
 #include "Widgets/Layout/SScrollBox.h"
-#include "Widgets/Layout/SUniformGridPanel.h"
 #include "Widgets/SBoxPanel.h"
 #include "Widgets/SNullWidget.h"
 
-DEFINE_LOG_CATEGORY_STATIC(LogCustomDialog, Log, All);
-
 SCustomDialog::FArguments& SCustomDialog::FArguments::IconBrush(FName InIconBrush)
 {
 	const FSlateBrush* ImageBrush = FAppStyle::Get().GetBrush(InIconBrush);
@@ -30,7 +26,6 @@ SCustomDialog::FArguments& SCustomDialog::FArguments::IconBrush(FName InIconBrus
 
 void SCustomDialog::Construct(const FArguments& InArgs)
 {
-	UE_LOG(LogCustomDialog, Log, TEXT("Dialog displayed:"), *InArgs._Title.ToString());
 	check(InArgs._Buttons.Num() > 0);
 	
 	OnClosed = InArgs._OnClosed;
@@ -109,7 +104,7 @@ TSharedRef<SWidget> SCustomDialog::CreateContentBox(const FArguments& InArgs)
 			.MaxDesiredHeight(InArgs._ScrollBoxMaxHeight)
 			[
 				SNew(SScrollBox)
-				+SScrollBox::Slot()
+				+ SScrollBox::Slot()
 				[
 					InArgs._Content.Widget
 				]
@@ -133,12 +128,12 @@ TSharedRef<SWidget> SCustomDialog::CreateContentBox(const FArguments& InArgs)
 
 TSharedRef<SWidget> SCustomDialog::CreateButtonBox(const FArguments& InArgs)
 {
-	TSharedPtr<SUniformGridPanel> ButtonPanel;
+	TSharedPtr<SHorizontalBox> ButtonPanel;
 	TSharedRef<SHorizontalBox> ButtonBox =
 		SNew(SHorizontalBox)
 	
 		// Before buttons
-		+SHorizontalBox::Slot()
+		+ SHorizontalBox::Slot()
 			.AutoWidth()
 			.HAlign(HAlign_Left)
 			.VAlign(VAlign_Center)
@@ -147,15 +142,12 @@ TSharedRef<SWidget> SCustomDialog::CreateButtonBox(const FArguments& InArgs)
 			]
 
 		// Buttons
-		+SHorizontalBox::Slot()
+		+ SHorizontalBox::Slot()
 			.FillWidth(1.0f)
 			.VAlign(VAlign_Center)
 			.HAlign(HAlign_Right)
 			[
-				SAssignNew(ButtonPanel, SUniformGridPanel)
-				.SlotPadding(FAppStyle::Get().GetMargin("StandardDialog.SlotPadding"))
-				.MinDesiredSlotWidth(FAppStyle::Get().GetFloat("StandardDialog.MinDesiredSlotWidth"))
-				.MinDesiredSlotHeight(FAppStyle::Get().GetFloat("StandardDialog.MinDesiredSlotHeight"))
+				SAssignNew(ButtonPanel, SHorizontalBox)
 			];
 
 	bool bCanFocusLastPrimary = true;
@@ -167,21 +159,18 @@ TSharedRef<SWidget> SCustomDialog::CreateButtonBox(const FArguments& InArgs)
 		const FButtonStyle* ButtonStyle = Button.bIsPrimary ?
 				&FAppStyle::Get().GetWidgetStyle< FButtonStyle >( "PrimaryButton" ) :
 				&FAppStyle::Get().GetWidgetStyle< FButtonStyle >("Button");
+		
 		TSharedRef<SButton> ButtonWidget = SNew(SButton)
 			.OnClicked(FOnClicked::CreateSP(this, &SCustomDialog::OnButtonClicked, Button.OnClicked, ButtonIndex))
 			.ButtonStyle(ButtonStyle)
 			[
-				SNew(SHorizontalBox)
-				+ SHorizontalBox::Slot()
-				.VAlign(VAlign_Center)
-				.HAlign(HAlign_Center)
-				[
-					SNew(STextBlock)
-					.Text(Button.ButtonText)
-				]
+				SNew(STextBlock)
+				.Text(Button.ButtonText)
 			];
 
-		ButtonPanel->AddSlot(ButtonPanel->GetChildren()->Num(), 0)
+		ButtonPanel->AddSlot()
+			.Padding(FAppStyle::Get().GetMargin("StandardDialog.SlotPadding"))
+			.AutoWidth()
 		[
 			ButtonWidget
 		];
