Make SCustomDialog buttons grid AutoWidth so the button text doesn't get clamped

diff --git a/Engine/Source/Developer/ToolWidgets/Private/Dialog/SCustomDialog.cpp b/Engine/Source/Developer/ToolWidgets/Private/Dialog/SCustomDialog.cpp
--- a/Engine/Source/Developer/ToolWidgets/Private/Dialog/SCustomDialog.cpp
+++ b/Engine/Source/Developer/ToolWidgets/Private/Dialog/SCustomDialog.cpp
@@ -143,7 +143,7 @@ TSharedRef<SWidget> SCustomDialog::CreateButtonBox(const FArguments& InArgs)
 
 		// Buttons
 		+ SHorizontalBox::Slot()
-			.FillWidth(1.0f)
+			.AutoWidth()
 			.VAlign(VAlign_Center)
 			.HAlign(HAlign_Right)
 			[
