Fixed SCustomDialog example comment that was using a deprecated API

#rnx

diff --git a/Engine/Source/Developer/ToolWidgets/Public/Dialog/SCustomDialog.h b/Engine/Source/Developer/ToolWidgets/Public/Dialog/SCustomDialog.h
--- a/Engine/Source/Developer/ToolWidgets/Public/Dialog/SCustomDialog.h
+++ b/Engine/Source/Developer/ToolWidgets/Public/Dialog/SCustomDialog.h
@@ -14,7 +14,10 @@
  * Usage:
  * TSharedRef<SCustomDialog> HelloWorldDialog = SNew(SCustomDialog)
 		.Title(FText(LOCTEXT("HelloWorldTitleExample", "Hello, World!")))
-		.DialogContent( SNew(SImage).Image(FName(TEXT("Hello"))))
+		.Content()
+		[
+			SNew(SImage).Image(FName(TEXT("Hello"))))
+		]
 		.Buttons({
 			SCustomDialog::FButton(LOCTEXT("OK", "OK")),
 			SCustomDialog::FButton(LOCTEXT("Cancel", "Cancel"))
