Fixed SCustomDialog buttons not being right aligned

Reverting a change that was made against 5.0, that isn't needed in 5.1 and causes this issue

#rnx

diff --git a/Engine/Source/Developer/ToolWidgets/Private/Dialog/SCustomDialog.cpp b/Engine/Source/Developer/ToolWidgets/Private/Dialog/SCustomDialog.cpp
--- a/Engine/Source/Developer/ToolWidgets/Private/Dialog/SCustomDialog.cpp
+++ b/Engine/Source/Developer/ToolWidgets/Private/Dialog/SCustomDialog.cpp
@@ -143,7 +143,7 @@ TSharedRef<SWidget> SCustomDialog::CreateButtonBox(const FArguments& InArgs)
 
 		// Buttons
 		+ SHorizontalBox::Slot()
-			.AutoWidth()
+			.FillWidth(1.0f)
 			.VAlign(VAlign_Center)
 			.HAlign(HAlign_Right)
 			[
