Add customizability of BeforeButton alignment for the SCustomDialog widget

diff --git a/Engine/Source/Developer/ToolWidgets/Private/Dialog/SCustomDialog.cpp b/Engine/Source/Developer/ToolWidgets/Private/Dialog/SCustomDialog.cpp
--- a/Engine/Source/Developer/ToolWidgets/Private/Dialog/SCustomDialog.cpp
+++ b/Engine/Source/Developer/ToolWidgets/Private/Dialog/SCustomDialog.cpp
@@ -50,7 +50,8 @@ void SCustomDialog::Construct(const FArguments& InArgs)
 				]
 				
 				+ SVerticalBox::Slot()
-				.VAlign(VAlign_Center)
+				.VAlign(InArgs._VAlignButtonBox)
+				.HAlign(InArgs._HAlignButtonBox)
 				.AutoHeight()
 				.Padding(InArgs._ButtonAreaPadding)
 				[
diff --git a/Engine/Source/Developer/ToolWidgets/Public/Dialog/SCustomDialog.h b/Engine/Source/Developer/ToolWidgets/Public/Dialog/SCustomDialog.h
--- a/Engine/Source/Developer/ToolWidgets/Public/Dialog/SCustomDialog.h
+++ b/Engine/Source/Developer/ToolWidgets/Public/Dialog/SCustomDialog.h
@@ -112,6 +112,11 @@ public:
 		/** Custom widget placed before the buttons */
 		SLATE_NAMED_SLOT(FArguments, BeforeButtons)
 
+		/** HAlign to use for Button Box slot (default: HAlign_Left) */
+		SLATE_ARGUMENT(EHorizontalAlignment, HAlignButtonBox)
+
+		/** VAlign to use for Button Box  slot (default: VAlign_Center) */
+		SLATE_ARGUMENT(EVerticalAlignment, VAlignButtonBox)
 
 		/** Padding to apply to the widget embedded in the window, i.e. to all widgets contained in the window (default: {4,4,4,4} )*/
 		SLATE_ATTRIBUTE(FMargin, RootPadding)
