Added option to display filter categories as sections rather than submenus to SBasicFilterBar

diff --git a/Engine/Source/Developer/ToolWidgets/Public/Filters/SBasicFilterBar.h b/Engine/Source/Developer/ToolWidgets/Public/Filters/SBasicFilterBar.h
--- a/Engine/Source/Developer/ToolWidgets/Public/Filters/SBasicFilterBar.h
+++ b/Engine/Source/Developer/ToolWidgets/Public/Filters/SBasicFilterBar.h
@@ -93,6 +93,7 @@ public:
  		: _FilterBarLayout(EFilterBarLayout::Horizontal)
 		, _CanChangeOrientation(false)
 		, _FilterPillStyle(EFilterPillStyle::Default)
+		, _UseSectionsForCategories(false)
 	{}
 
  		/** Delegate for when filters have changed */
@@ -124,6 +125,9 @@ public:
 
 		/** Determines how each individual filter pill looks like */
 		SLATE_ARGUMENT(EFilterPillStyle, FilterPillStyle)
+
+		/** Whether to use submenus or sections for categories in the filter menu */
+		SLATE_ARGUMENT(bool, UseSectionsForCategories)
 	
  	SLATE_END_ARGS()
 
@@ -135,6 +139,7 @@ public:
  		FilterBarLayout = InArgs._FilterBarLayout;
  		bCanChangeOrientation = InArgs._CanChangeOrientation;
  		FilterPillStyle = InArgs._FilterPillStyle;
+		bUseSectionsForCategories = InArgs._UseSectionsForCategories;
 
  		// We use a widgetswitcher to allow dynamically swapping between the layouts
  		ChildSlot
@@ -1523,22 +1528,34 @@ protected:
 	/** Helper function to add all custom filters to the Add Filter Menu */
 	void PopulateCustomFilters(UToolMenu* Menu)
 	{
-		FToolMenuSection& Section = Menu->AddSection("BasicFilterBarFiltersMenu", LOCTEXT("FilterBarOtherFiltersSection", "Other Filters"));
-		
-		// Add all the filters
-		for (const TSharedPtr<FFilterCategory>& Category : AllFilterCategories)
+		if(bUseSectionsForCategories)
 		{
-			Section.AddSubMenu(
-				NAME_None,
-				Category->Title,
-				Category->Tooltip,
-				FNewToolMenuDelegate::CreateSP(this, &SBasicFilterBar<FilterType>::CreateOtherFiltersMenuCategory, Category),
-				FUIAction(
-				FExecuteAction::CreateSP( this, &SBasicFilterBar<FilterType>::FrontendFilterCategoryClicked, Category ),
-				FCanExecuteAction(),
-				FGetActionCheckState::CreateSP(this, &SBasicFilterBar<FilterType>::IsFrontendFilterCategoryChecked, Category ) ),
-				EUserInterfaceActionType::ToggleButton
-				);
+			// Add all the filters as sections
+			for (const TSharedPtr<FFilterCategory>& Category : AllFilterCategories)
+			{
+				FToolMenuSection& Section = Menu->AddSection(*Category->Title.ToString(), Category->Title);
+				CreateOtherFiltersMenuCategory(Section, Category);
+			}
+		}
+		else
+		{
+			FToolMenuSection& Section = Menu->AddSection("BasicFilterBarFiltersMenu", LOCTEXT("FilterBarOtherFiltersSection", "Other Filters"));
+			
+			// Add all the filters as submenus
+			for (const TSharedPtr<FFilterCategory>& Category : AllFilterCategories)
+			{
+				Section.AddSubMenu(
+					NAME_None,
+					Category->Title,
+					Category->Tooltip,
+					FNewToolMenuDelegate::CreateSP(this, &SBasicFilterBar<FilterType>::CreateOtherFiltersMenuCategory, Category),
+					FUIAction(
+					FExecuteAction::CreateSP( this, &SBasicFilterBar<FilterType>::FrontendFilterCategoryClicked, Category ),
+					FCanExecuteAction(),
+					FGetActionCheckState::CreateSP(this, &SBasicFilterBar<FilterType>::IsFrontendFilterCategoryChecked, Category ) ),
+					EUserInterfaceActionType::ToggleButton
+					);
+			}
 		}
 	}
 
@@ -1622,6 +1639,9 @@ protected:
 	/** Whether the orientation can be changed after initailization */
 	bool bCanChangeOrientation;
 
+	/** Whether to use submenus or sections for categories in the filter menu */
+	bool bUseSectionsForCategories;
+
 	/** Determines how each individual pill looks like */
 	EFilterPillStyle FilterPillStyle;
 	
