Fix uninitialized variable in WidgetDrawer.

diff --git a/Engine/Source/Developer/ToolWidgets/Public/WidgetDrawerConfig.h b/Engine/Source/Developer/ToolWidgets/Public/WidgetDrawerConfig.h
--- a/Engine/Source/Developer/ToolWidgets/Public/WidgetDrawerConfig.h
+++ b/Engine/Source/Developer/ToolWidgets/Public/WidgetDrawerConfig.h
@@ -24,7 +24,7 @@ struct FWidgetDrawerConfig
 	TSharedPtr<SWidget> CustomWidget;
 	FText ButtonText;
 	FText ToolTipText;
-	const FSlateBrush* Icon;
+	const FSlateBrush* Icon = nullptr;
 
 	bool operator==(const FName& OtherId) const
 	{
