SWidgetDrawer: Collapse custom widget border if said widget is collapsed

diff --git a/Engine/Source/Developer/ToolWidgets/Private/SWidgetDrawer.cpp b/Engine/Source/Developer/ToolWidgets/Private/SWidgetDrawer.cpp
--- a/Engine/Source/Developer/ToolWidgets/Private/SWidgetDrawer.cpp
+++ b/Engine/Source/Developer/ToolWidgets/Private/SWidgetDrawer.cpp
@@ -13,6 +13,7 @@
 #include "Widgets/Notifications/SProgressBar.h"
 #include "Widgets/Notifications/SNotificationList.h"
 #include "Framework/Notifications/NotificationManager.h"
+#include "Types/SlateAttributeMetaData.h"
 
 #define LOCTEXT_NAMESPACE "StatusBar"
 
@@ -477,20 +478,33 @@ TSharedRef<SWidget> SWidgetDrawer::MakeStatusBarDrawerButton(const FWidgetDrawer
 
 	if (Drawer.CustomWidget)
 	{
+		auto IsCustomWidgetBorderVisible = [CustomWidgetWeak = Drawer.CustomWidget.ToWeakPtr()]()
+		{
+			if (TSharedPtr<SWidget> CustomWidget = CustomWidgetWeak.Pin())
+			{
+				FSlateAttributeMetaData::UpdateOnlyVisibilityAttributes(*CustomWidget, FSlateAttributeMetaData::EInvalidationPermission::AllowInvalidationIfConstructed);
+				if (CustomWidget->GetVisibility() == EVisibility::Collapsed)
+				{
+					return EVisibility::Collapsed;
+				}
+			}
+			return EVisibility::SelfHitTestInvisible;
+		};
+
 		return
 			SNew(SHorizontalBox)
 			+ SHorizontalBox::Slot()
-			.Padding(0.0f, 0.0f, 2.0f, 0.0f)
 			.AutoWidth()
 			[
 				DrawerButton
 			]
 			+ SHorizontalBox::Slot()
+			.Padding(2.0f, 0.0f, 0.0f, 0.0f)
 			[
 				SNew(SBorder)
 				.Padding(FMargin(2.0f, 0.0f))
 				.BorderImage(StatusBarBackground)
-				.Visibility(EVisibility::SelfHitTestInvisible)
+				.Visibility(MakeAttributeLambda(IsCustomWidgetBorderVisible))
 				.VAlign(VAlign_Center)
 				[
 					Drawer.CustomWidget.ToSharedRef()
