Collapse the SMessageDialog icon padding when no icon is set and fixed content being right aligned

#rnx

diff --git a/Engine/Source/Developer/ToolWidgets/Private/Dialog/SCustomDialog.cpp b/Engine/Source/Developer/ToolWidgets/Private/Dialog/SCustomDialog.cpp
--- a/Engine/Source/Developer/ToolWidgets/Private/Dialog/SCustomDialog.cpp
+++ b/Engine/Source/Developer/ToolWidgets/Private/Dialog/SCustomDialog.cpp
@@ -44,7 +44,6 @@ void SCustomDialog::Construct(const FArguments& InArgs)
 				SNew(SVerticalBox)
 				+ SVerticalBox::Slot()
 				.FillHeight(1.0f)
-				.HAlign(HAlign_Right)
 				[
 					CreateContentBox(InArgs)
 				]
@@ -84,14 +83,15 @@ TSharedRef<SWidget> SCustomDialog::CreateContentBox(const FArguments& InArgs)
 		.AutoWidth()
 		.VAlign(InArgs._VAlignIcon)
 		.HAlign(InArgs._HAlignIcon)
-		.Padding(0, 0, 8, 0)
 		[
-			!InArgs._Icon.IsSet() || !InArgs._Icon.Get()
-				// If no icon is specified, we still want to pad the content towards the right
-				? SNullWidget::NullWidget
-				: SNew(SImage)
+			SNew(SBox)
+			.Padding(0, 0, 8, 0)
+			.Visibility_Lambda([IconAttr = InArgs._Icon]() { return IconAttr.Get(nullptr) ? EVisibility::Visible : EVisibility::Collapsed; })
+			[
+				SNew(SImage)
 				.DesiredSizeOverride(InArgs._IconDesiredSizeOverride)
 				.Image(InArgs._Icon)
+			]
 		];
 
 	if (InArgs._UseScrollBox)
