Status Bar:
- Added option to unregister a Drawer

diff --git a/Engine/Source/Developer/ToolWidgets/Private/SWidgetDrawer.cpp b/Engine/Source/Developer/ToolWidgets/Private/SWidgetDrawer.cpp
--- a/Engine/Source/Developer/ToolWidgets/Private/SWidgetDrawer.cpp
+++ b/Engine/Source/Developer/ToolWidgets/Private/SWidgetDrawer.cpp
@@ -577,15 +577,37 @@ void SWidgetDrawer::RegisterDrawer(FWidgetDrawerConfig&& Drawer, int32 SlotIndex
 
 	if (RegisteredDrawers.Num() > NumDrawers)
 	{
+		TSharedRef<SWidget> Content = MakeStatusBarDrawerButton(Drawer);
+
+		DrawerIdToContentWidget.Add(Drawer.UniqueId, Content);
+
 		DrawerBox->InsertSlot(SlotIndex)
 		.Padding(1.0f, 0.0f)
 		.AutoWidth()
 		[
-			MakeStatusBarDrawerButton(Drawer)
+			Content
 		];
 	}
 }
 
+void SWidgetDrawer::UnregisterDrawer(FName DrawerId)
+{
+	if (IsDrawerOpened(DrawerId))
+	{
+		CloseDrawerImmediately(DrawerId);
+	}
+
+	RegisteredDrawers.Remove(DrawerId);
+
+	TWeakPtr<SWidget> ContentWidgetWeak;
+	DrawerIdToContentWidget.RemoveAndCopyValue(DrawerId, ContentWidgetWeak);
+
+	if (TSharedPtr<SWidget> ContentWidget = ContentWidgetWeak.Pin())
+	{
+		DrawerBox->RemoveSlot(ContentWidget.ToSharedRef());
+	}
+}
+
 void SWidgetDrawer::OpenDrawer(const FName DrawerId)
 {
 	// Close any other open drawer
diff --git a/Engine/Source/Developer/ToolWidgets/Public/SWidgetDrawer.h b/Engine/Source/Developer/ToolWidgets/Public/SWidgetDrawer.h
--- a/Engine/Source/Developer/ToolWidgets/Public/SWidgetDrawer.h
+++ b/Engine/Source/Developer/ToolWidgets/Public/SWidgetDrawer.h
@@ -55,6 +55,14 @@ public:
 	 */
 	void RegisterDrawer(FWidgetDrawerConfig&& Drawer, int32 SlotIndex = INDEX_NONE);
 
+	/**
+	 * Unregisters drawer with the given DrawerId.
+	 * This will dismiss the drawer immediately if its opened, and destroy the status bar widget that opens/closes the drawer
+	 *
+	 * @param DrawerId name of the registered drawer to unregister 
+	 */
+	void UnregisterDrawer(FName DrawerId);
+
 	/**
 	 * Opens a drawer
 	 * 
@@ -102,6 +110,9 @@ private:
 private:
 
 	TArray<FWidgetDrawerConfig> RegisteredDrawers;
+
+	TMap<FName, TWeakPtr<SWidget>> DrawerIdToContentWidget;
+
 	TArray<FOpenDrawerData> DismissingDrawers;
 	
 	TSharedPtr<SHorizontalBox> DrawerBox;
diff --git a/Engine/Source/Editor/StatusBar/Private/SStatusBar.cpp b/Engine/Source/Editor/StatusBar/Private/SStatusBar.cpp
--- a/Engine/Source/Editor/StatusBar/Private/SStatusBar.cpp
+++ b/Engine/Source/Editor/StatusBar/Private/SStatusBar.cpp
@@ -800,6 +800,11 @@ void SStatusBar::RegisterDrawer(FWidgetDrawerConfig&& Drawer, int32 SlotIndex)
 	WidgetDrawer->RegisterDrawer(MoveTemp(Drawer), SlotIndex);
 }
 
+void SStatusBar::UnregisterDrawer(FName DrawerId)
+{
+	WidgetDrawer->UnregisterDrawer(DrawerId);
+}
+
 void SStatusBar::OpenDrawer(const FName DrawerId)
 {
 	WidgetDrawer->OpenDrawer(DrawerId);
diff --git a/Engine/Source/Editor/StatusBar/Private/SStatusBar.h b/Engine/Source/Editor/StatusBar/Private/SStatusBar.h
--- a/Engine/Source/Editor/StatusBar/Private/SStatusBar.h
+++ b/Engine/Source/Editor/StatusBar/Private/SStatusBar.h
@@ -122,6 +122,14 @@ public:
 	 */
 	void RegisterDrawer(FWidgetDrawerConfig&& Drawer, int32 SlotIndex = INDEX_NONE);
 
+	/**
+	 * Unregisters drawer with the given DrawerId.
+	 * This will dismiss the drawer immediately if its opened, and destroy the status bar widget that opens/closes the drawer
+	 *
+	 * @param DrawerId name of the registered drawer to unregister 
+	 */
+	void UnregisterDrawer(FName DrawerId);
+
 	/**
 	 * Opens a drawer
 	 * 
diff --git a/Engine/Source/Editor/StatusBar/Private/StatusBarSubsystem.cpp b/Engine/Source/Editor/StatusBar/Private/StatusBarSubsystem.cpp
--- a/Engine/Source/Editor/StatusBar/Private/StatusBarSubsystem.cpp
+++ b/Engine/Source/Editor/StatusBar/Private/StatusBarSubsystem.cpp
@@ -582,6 +582,14 @@ void UStatusBarSubsystem::RegisterDrawer(FName StatusBarName, FWidgetDrawerConfi
 	}
 }
 
+void UStatusBarSubsystem::UnregisterDrawer(FName StatusBarName, FName DrawerId)
+{
+	if (TSharedPtr<SStatusBar> StatusBar = GetStatusBar(StatusBarName))
+	{
+		StatusBar->UnregisterDrawer(DrawerId);
+	}
+}
+
 FStatusBarMessageHandle UStatusBarSubsystem::PushStatusBarMessage(FName StatusBarName, const TAttribute<FText>& InMessage, const TAttribute<FText>& InHintText)
 {
 	if (TSharedPtr<SStatusBar> StatusBar = GetStatusBar(StatusBarName))
diff --git a/Engine/Source/Editor/StatusBar/Public/StatusBarSubsystem.h b/Engine/Source/Editor/StatusBar/Public/StatusBarSubsystem.h
--- a/Engine/Source/Editor/StatusBar/Public/StatusBarSubsystem.h
+++ b/Engine/Source/Editor/StatusBar/Public/StatusBarSubsystem.h
@@ -129,6 +129,14 @@ public:
 	 */
 	void RegisterDrawer(FName StatusBarName, FWidgetDrawerConfig&& Drawer, int32 SlotIndex = INDEX_NONE);
 
+	/**
+	 * Unregisters and destroys the drawer with the given DrawerId 
+	 *
+	 * @param StatusBarName	The name of the status bar to unregister the drawer from
+	 * @param DrawerId		the unique name id of the drawer to unregister
+	 */
+	void UnregisterDrawer(FName StatusBarName, FName DrawerId);
+
 	/** 
 	 * Pushes a new status bar message
 	 *
