Creating a small text variant of SSimpleComboButton

diff --git a/Engine/Source/Developer/ToolWidgets/Private/SSimpleComboButton.cpp b/Engine/Source/Developer/ToolWidgets/Private/SSimpleComboButton.cpp
--- a/Engine/Source/Developer/ToolWidgets/Private/SSimpleComboButton.cpp
+++ b/Engine/Source/Developer/ToolWidgets/Private/SSimpleComboButton.cpp
@@ -4,12 +4,13 @@
 #include "Widgets/SBoxPanel.h"
 #include "Widgets/Images/SImage.h"
 #include "Widgets/Text/STextBlock.h"
+#include "Styling/SlateTypes.h"
 
 void SSimpleComboButton::Construct(const FArguments& InArgs)
 {
 
 	TAttribute<FText> Text = InArgs._Text;
-
+	FTextBlockStyle TextStyle = InArgs._UsesSmallText ? FAppStyle::GetWidgetStyle<FTextBlockStyle>("SmallText") : FAppStyle::GetWidgetStyle<FTextBlockStyle>("SmallButtonText");
 
 	TSharedRef<SHorizontalBox> ButtonContent = SNew(SHorizontalBox)
 		+ SHorizontalBox::Slot()
@@ -24,12 +25,13 @@ void SSimpleComboButton::Construct(const FArguments& InArgs)
 		+ SHorizontalBox::Slot()
 		.Padding(FMargin(3, 0, 0, 0))
 		.VAlign(VAlign_Center)
-		.AutoWidth()
+		.FillWidth(1.0f)
 		[
 			SNew(STextBlock)
-			.TextStyle(FAppStyle::Get(), "SmallButtonText")
+			.TextStyle(&TextStyle)
 			.Text(InArgs._Text)
 			.Visibility_Lambda([Text]() { return Text.Get(FText::GetEmpty()).IsEmpty() ? EVisibility::Collapsed : EVisibility::Visible; })
+			.Clipping(EWidgetClipping::OnDemand)
 		];
 
 	SComboButton::Construct(SComboButton::FArguments()
diff --git a/Engine/Source/Developer/ToolWidgets/Public/SSimpleComboButton.h b/Engine/Source/Developer/ToolWidgets/Public/SSimpleComboButton.h
--- a/Engine/Source/Developer/ToolWidgets/Public/SSimpleComboButton.h
+++ b/Engine/Source/Developer/ToolWidgets/Public/SSimpleComboButton.h
@@ -12,6 +12,7 @@ public:
 
 	SLATE_BEGIN_ARGS(SSimpleComboButton)
 		: _HasDownArrow(false)
+		, _UsesSmallText(false)
 	{}
 		/** The text to display in the button. */
 		SLATE_ATTRIBUTE(FText, Text)
@@ -25,6 +26,7 @@ public:
 		SLATE_EVENT(FOnComboBoxOpened, OnComboBoxOpened)
 		SLATE_EVENT(FOnIsOpenChanged, OnMenuOpenChanged)
 		SLATE_ARGUMENT(bool, HasDownArrow)
+		SLATE_ARGUMENT(bool, UsesSmallText)
 
 	SLATE_END_ARGS()
 
