SCheckBoxList Callback for when checkboxes are changed

diff --git a/Engine/Source/Developer/ToolWidgets/Private/SCheckBoxList.cpp b/Engine/Source/Developer/ToolWidgets/Private/SCheckBoxList.cpp
--- a/Engine/Source/Developer/ToolWidgets/Private/SCheckBoxList.cpp
+++ b/Engine/Source/Developer/ToolWidgets/Private/SCheckBoxList.cpp
@@ -102,6 +102,7 @@ void SCheckBoxList::Construct(const FArguments& InArgs, const TArray<TSharedRef<
 	}
 
 	bool bShowHeaderCheckbox = InArgs._IncludeGlobalCheckBoxInHeaderRow;
+	OnItemCheckStateChanged = InArgs._OnItemCheckStateChanged;
 
 	TSharedRef<SHeaderRow> HeaderRowWidget = SNew(SHeaderRow)
 		+ SHeaderRow::Column(CheckBoxList::ColumnID_CheckBox)
@@ -191,13 +192,21 @@ void SCheckBoxList::OnToggleSelectedCheckBox(ECheckBoxState InNewState)
 		Item->bIsChecked = bNewValue;
 	}
 	bAllCheckedState = bNewValue ? ECheckBoxState::Checked : ECheckBoxState::Unchecked;
+	
+	OnItemCheckStateChanged.ExecuteIfBound(-1);
+}
+
+void SCheckBoxList::OnItemCheckBox(TSharedRef<CheckBoxList::FItemPair> InItem)
+{
+	UpdateAllChecked();
+	OnItemCheckStateChanged.ExecuteIfBound(Items.IndexOfByKey(InItem));
 }
 
 TSharedRef<ITableRow> SCheckBoxList::HandleGenerateRow(TSharedRef<CheckBoxList::FItemPair> InItem, const TSharedRef<STableViewBase>& OwnerTable)
 {
 	return SNew(CheckBoxList::SItemPair, OwnerTable, InItem)
 		.CheckBoxStyle(CheckBoxStyle)
-		.OnCheckUpdated(FSimpleDelegate::CreateSP(this, &SCheckBoxList::UpdateAllChecked));
+		.OnCheckUpdated(FSimpleDelegate::CreateLambda([this, InItem]() { OnItemCheckBox(InItem); }));
 }
 
 #undef LOCTEXT_NAMESPACE
\ No newline at end of file
diff --git a/Engine/Source/Developer/ToolWidgets/Public/SCheckBoxList.h b/Engine/Source/Developer/ToolWidgets/Public/SCheckBoxList.h
--- a/Engine/Source/Developer/ToolWidgets/Public/SCheckBoxList.h
+++ b/Engine/Source/Developer/ToolWidgets/Public/SCheckBoxList.h
@@ -22,6 +22,8 @@ namespace CheckBoxList
 	struct FItemPair;
 }
 
+DECLARE_DELEGATE_OneParam( FOnCheckListItemStateChanged, int );
+
 /** A widget that can be used inside a CustomDialog to display a list of checkboxes */
 class TOOLWIDGETS_API SCheckBoxList: public SCompoundWidget
 {
@@ -33,8 +35,10 @@ class TOOLWIDGETS_API SCheckBoxList: public SCompoundWidget
 		SLATE_STYLE_ARGUMENT(FCheckBoxStyle, CheckBoxStyle)
 		/** The label of the item column header */
 		SLATE_ARGUMENT(FText, ItemHeaderLabel)
-		/** The label of the item column header */
+		/** Optionally display a checkbox by the column header that toggles all items */
 		SLATE_ARGUMENT(bool, IncludeGlobalCheckBoxInHeaderRow)
+		/** Callback when any checkbox is changed. Parameter is the index of the item, or -1 if it was the "All"/Global checkbox */
+		SLATE_EVENT( FOnCheckListItemStateChanged, OnItemCheckStateChanged )
 	SLATE_END_ARGS()
 
 public:
@@ -54,6 +58,7 @@ private:
 	void UpdateAllChecked();
 	ECheckBoxState GetToggleSelectedState() const;
 	void OnToggleSelectedCheckBox(ECheckBoxState InNewState);
+	void OnItemCheckBox(TSharedRef<CheckBoxList::FItemPair> InItem);
 
 	TSharedRef<ITableRow> HandleGenerateRow(TSharedRef<CheckBoxList::FItemPair> InItem, const TSharedRef<STableViewBase>& OwnerTable);
 
@@ -62,4 +67,6 @@ private:
 
 	const FCheckBoxStyle* CheckBoxStyle;
 	TSharedPtr<SListView<TSharedRef<CheckBoxList::FItemPair>>> ListView;
+
+	FOnCheckListItemStateChanged OnItemCheckStateChanged;
 };
