Add tokenized search support for Searchable Combo Box

diff --git a/Engine/Source/Developer/ToolWidgets/Private/SSearchableComboBox.cpp b/Engine/Source/Developer/ToolWidgets/Private/SSearchableComboBox.cpp
--- a/Engine/Source/Developer/ToolWidgets/Private/SSearchableComboBox.cpp
+++ b/Engine/Source/Developer/ToolWidgets/Private/SSearchableComboBox.cpp
@@ -137,9 +137,22 @@ void SSearchableComboBox::RefreshOptions()
 	}
 	else
 	{
-		for (TSharedPtr<FString> Option : *OptionsSource)
+		TArray<FString> SearchTokens;
+		SearchText.ToString().ParseIntoArrayWS(SearchTokens);
+
+		for (const TSharedPtr<FString>& Option : *OptionsSource)
 		{
-			if (Option->Find(SearchText.ToString(), ESearchCase::Type::IgnoreCase) >= 0)
+			bool bAllTokensMatch = true;
+			for (const FString& SearchToken : SearchTokens)
+			{
+				if (Option->Find(SearchToken, ESearchCase::Type::IgnoreCase) == INDEX_NONE)
+				{
+					bAllTokensMatch = false;
+					break;
+				}
+			}
+
+			if (bAllTokensMatch)
 			{
 				FilteredOptionsSource.Add(Option);
 			}
