[ContentBrowser/ToolWidgets]
Moving SSearchToggleButton into ToolWidgets module to be reused externally.

diff --git a/Engine/Source/Developer/ToolWidgets/Private/SSearchToggleButton.cpp b/Engine/Source/Developer/ToolWidgets/Private/SSearchToggleButton.cpp
new file mode 100644
--- /dev/null
+++ b/Engine/Source/Developer/ToolWidgets/Private/SSearchToggleButton.cpp
@@ -0,0 +1,77 @@
+// Copyright Epic Games, Inc. All Rights Reserved.
+
+#include "SSearchToggleButton.h"
+
+#include "Framework/Application/SlateApplication.h"
+#include "Input/Events.h"
+#include "Internationalization/Internationalization.h"
+#include "Internationalization/Text.h"
+#include "Layout/Children.h"
+#include "Layout/Margin.h"
+#include "Misc/Attribute.h"
+#include "Styling/SlateColor.h"
+#include "Widgets/Images/SImage.h"
+#include "Widgets/Input/SCheckBox.h"
+#include "Widgets/Input/SSearchBox.h"
+
+void SSearchToggleButton::Construct(const FArguments& InArgs, TSharedRef<SSearchBox> SearchBox)
+{
+	bIsExpanded = false;
+	OnSearchBoxShown = InArgs._OnSearchBoxShown;
+
+	SearchStyle = InArgs._Style;
+
+	SearchBox->SetVisibility(TAttribute<EVisibility>::CreateSP(this, &SSearchToggleButton::GetSearchBoxVisibility));
+	SearchBoxPtr = SearchBox;
+
+	ChildSlot
+	[
+		SNew(SCheckBox)
+		.IsChecked(this, &SSearchToggleButton::GetToggleButtonState)
+		.OnCheckStateChanged(this, &SSearchToggleButton::OnToggleButtonStateChanged)
+		.Style(FAppStyle::Get(), "ToggleButtonCheckbox")
+		.Padding(4.0f)
+		.ToolTipText(NSLOCTEXT("ExpandableSearchArea", "ExpandCollapseSearchButton", "Expands or collapses the search text box"))
+		[
+			SNew(SImage)
+			.Image(&SearchStyle->GlassImage)
+			.ColorAndOpacity(FSlateColor::UseForeground())
+		]
+	];
+}
+
+void SSearchToggleButton::SetExpanded(bool bInExpanded)
+{
+	bIsExpanded = bInExpanded;
+
+	if (TSharedPtr<SSearchBox> SearchBox = SearchBoxPtr.Pin())
+	{
+		if (bIsExpanded)
+		{
+			OnSearchBoxShown.ExecuteIfBound();
+
+			// Focus the search box when it's shown
+			FSlateApplication::Get().SetUserFocus(FSlateApplication::Get().GetUserIndexForKeyboard(), SearchBox, EFocusCause::SetDirectly);
+		}
+		else
+		{
+			// Clear the search box when it's hidden
+			SearchBox->SetText(FText::GetEmpty());
+		}
+	}
+}
+
+void SSearchToggleButton::OnToggleButtonStateChanged(ECheckBoxState CheckBoxState)
+{
+	SetExpanded(CheckBoxState == ECheckBoxState::Checked);
+}
+
+ECheckBoxState SSearchToggleButton::GetToggleButtonState() const
+{
+	return bIsExpanded ? ECheckBoxState::Checked : ECheckBoxState::Unchecked;
+}
+
+EVisibility SSearchToggleButton::GetSearchBoxVisibility() const
+{
+	return bIsExpanded ? EVisibility::Visible : EVisibility::Collapsed;
+}
diff --git a/Engine/Source/Developer/ToolWidgets/Public/SSearchToggleButton.h b/Engine/Source/Developer/ToolWidgets/Public/SSearchToggleButton.h
new file mode 100644
--- /dev/null
+++ b/Engine/Source/Developer/ToolWidgets/Public/SSearchToggleButton.h
@@ -0,0 +1,50 @@
+// Copyright Epic Games, Inc. All Rights Reserved.
+
+#pragma once 
+
+#include "Delegates/Delegate.h"
+#include "Layout/Visibility.h"
+#include "Styling/AppStyle.h"
+#include "Styling/ISlateStyle.h"
+#include "Styling/SlateTypes.h"
+#include "Styling/SlateWidgetStyleAsset.h"
+#include "Templates/SharedPointer.h"
+#include "Widgets/DeclarativeSyntaxSupport.h"
+#include "Widgets/SCompoundWidget.h"
+
+class SSearchBox;
+
+DECLARE_DELEGATE(FOnSearchBoxShown)
+
+class TOOLWIDGETS_API SSearchToggleButton : public SCompoundWidget
+{
+public:
+	SLATE_BEGIN_ARGS(SSearchToggleButton)
+		: _Style(&FAppStyle::Get().GetWidgetStyle<FSearchBoxStyle>("SearchBox"))
+	{}
+		/** Search box style (used to match the glass icon) */
+		SLATE_STYLE_ARGUMENT(FSearchBoxStyle, Style)
+
+		/** Event fired when the associated search box is made visible */
+		SLATE_EVENT(FOnSearchBoxShown, OnSearchBoxShown)
+
+	SLATE_END_ARGS()
+
+	void Construct(const FArguments& InArgs, TSharedRef<SSearchBox> SearchBox);
+
+	/** @return true if the search area is expanded and the search box exposed */
+	bool IsExpanded() const { return bIsExpanded; }
+
+	/** Sets whether or not the search area is expanded to expose the search box */
+	void SetExpanded(bool bInExpanded);
+private:
+
+	ECheckBoxState GetToggleButtonState() const;
+	void OnToggleButtonStateChanged(ECheckBoxState CheckBoxState);
+	EVisibility GetSearchBoxVisibility() const;
+private:
+	const FSearchBoxStyle* SearchStyle;
+	bool bIsExpanded;
+	FOnSearchBoxShown OnSearchBoxShown;
+	TWeakPtr<SSearchBox> SearchBoxPtr;
+};
diff --git a/Engine/Source/Editor/ContentBrowser/Private/SSearchToggleButton.cpp b/Engine/Source/Editor/ContentBrowser/Private/SSearchToggleButton.cpp
deleted file mode 100644
diff --git a/Engine/Source/Editor/ContentBrowser/Private/SSearchToggleButton.h b/Engine/Source/Editor/ContentBrowser/Private/SSearchToggleButton.h
deleted file mode 100644
