Fix incorrect SMessageDialog init params

diff --git a/Engine/Source/Developer/ToolWidgets/Private/Dialog/SMessageDialog.cpp b/Engine/Source/Developer/ToolWidgets/Private/Dialog/SMessageDialog.cpp
--- a/Engine/Source/Developer/ToolWidgets/Private/Dialog/SMessageDialog.cpp
+++ b/Engine/Source/Developer/ToolWidgets/Private/Dialog/SMessageDialog.cpp
@@ -35,7 +35,7 @@ void SMessageDialog::Construct(const FArguments& InArgs)
 		.VAlignIcon(VAlign_Top)
 		.ContentAreaPadding(FMargin(16.f, 0.f, 0.f, 0.f))
 		.UseScrollBox(InArgs._UseScrollBox)
-		.ScrollBoxMaxHeight(InArgs._UseScrollBox)
+		.ScrollBoxMaxHeight(InArgs._ScrollBoxMaxHeight)
 		.ButtonAreaPadding(0.f)
 		.OnClosed(InArgs._OnClosed)
 		.BeforeButtons()
