Don't close the content drawer if no new active tab is being set. This happens when closing a tab drawer where we won't have an active tab. Given that when we open the content drawer from a tab drawer, the tab drawer will lose focus and close itself, we don't want to also close the content drawer.

diff --git a/Engine/Source/Developer/ToolWidgets/Private/SWidgetDrawer.cpp b/Engine/Source/Developer/ToolWidgets/Private/SWidgetDrawer.cpp
--- a/Engine/Source/Developer/ToolWidgets/Private/SWidgetDrawer.cpp
+++ b/Engine/Source/Developer/ToolWidgets/Private/SWidgetDrawer.cpp
@@ -411,20 +411,18 @@ void SWidgetDrawer::OnGlobalFocusChanging(const FFocusEvent& FocusEvent, const F
 void SWidgetDrawer::OnActiveTabChanged(TSharedPtr<SDockTab> PreviouslyActive, TSharedPtr<SDockTab> NewlyActivated)
 {
 	bool bShouldRemoveDrawer = false;
-	if (!PreviouslyActive || !NewlyActivated)
+	if (NewlyActivated)
 	{
-		// Remove the content browser if there is some invalid state with the tabs
-		bShouldRemoveDrawer = true;
-	}
-	else if(NewlyActivated->GetTabRole() == ETabRole::MajorTab)
-	{
-		// Remove the content browser if a newly activated tab is a major tab
-		bShouldRemoveDrawer = true;
-	}
-	else if (PreviouslyActive->GetTabManagerPtr() != NewlyActivated->GetTabManagerPtr())
-	{
-		// Remove the content browser if we're switching tab managers (indicates a new status bar is becoming active)
-		bShouldRemoveDrawer = true;
+		if (NewlyActivated->GetTabRole() == ETabRole::MajorTab)
+		{
+			// Remove the drawer if a newly activated tab is a major tab
+			bShouldRemoveDrawer = true;
+		}
+		else if (PreviouslyActive && PreviouslyActive->GetTabManagerPtr() != NewlyActivated->GetTabManagerPtr())
+		{
+			// Remove the drawer if we're switching tab managers (indicates a new status bar is becoming active)
+			bShouldRemoveDrawer = true;
+		}
 	}
 
 	if (bShouldRemoveDrawer)
