Fixing missing header guard in DialogUtils header file (added in CL 28276965)

#rnx

diff --git a/Engine/Source/Developer/ToolWidgets/Public/Dialog/DialogUtils.h b/Engine/Source/Developer/ToolWidgets/Public/Dialog/DialogUtils.h
--- a/Engine/Source/Developer/ToolWidgets/Public/Dialog/DialogUtils.h
+++ b/Engine/Source/Developer/ToolWidgets/Public/Dialog/DialogUtils.h
@@ -1,5 +1,7 @@
 // Copyright Epic Games, Inc. All Rights Reserved.
 
+#pragma once
+
 #include "Internationalization/Text.h"
 
 enum class EAppMsgCategory : uint8;
