Fix some missing default values in SMessageDialog arguments
Notably, when the messag text got too long, because the ScrollBoxMaxHeight arg wasn't set, you could barely see the message
#rnx

diff --git a/Engine/Source/Developer/ToolWidgets/Public/Dialog/SMessageDialog.h b/Engine/Source/Developer/ToolWidgets/Public/Dialog/SMessageDialog.h
--- a/Engine/Source/Developer/ToolWidgets/Public/Dialog/SMessageDialog.h
+++ b/Engine/Source/Developer/ToolWidgets/Public/Dialog/SMessageDialog.h
@@ -17,7 +17,10 @@ public:
 
 	SLATE_BEGIN_ARGS(SMessageDialog)
 		: _AutoCloseOnButtonPress(true)
+		, _DecoratorStyleSet(nullptr)
 		, _Icon(nullptr)
+		, _UseScrollBox(true)
+		, _ScrollBoxMaxHeight(300)
 		, _WrapMessageAt(512.f)
 	{}
 
