- Renamed OnSetFilterActive parameter to reflect what it's actually doing (enabling/disabling), and can now also be used to enable a filter that had been pinned before.
- Additionally added 'IsFilterEnabled' function so that external systems can determine whether a filter is enabled or not

diff --git a/Engine/Source/Developer/ToolWidgets/Public/Filters/FilterBase.h b/Engine/Source/Developer/ToolWidgets/Public/Filters/FilterBase.h
--- a/Engine/Source/Developer/ToolWidgets/Public/Filters/FilterBase.h
+++ b/Engine/Source/Developer/ToolWidgets/Public/Filters/FilterBase.h
@@ -66,8 +66,19 @@ public:
 	/** Can be overriden for custom FilterBar subclasses to load settings, currently not implemented in any gneeric Filter Bar */
 	virtual void LoadSettings(const FString& IniFilename, const FString& IniSection, const FString& SettingsString) = 0;
 
-	/** Set this filter as active/inactive */
-	void SetActive(bool bInActive) { SetActiveEvent.Broadcast(bInActive); }
+	/** Set this filter as checked, and activates or deactivates it */
+	void SetActive(bool bActive) { SetActiveEvent.Broadcast(bActive); }
+
+	/** Checks whether the filter is checked/pinned and active/enabled */
+	bool IsActive() const
+	{
+		if(IsActiveEvent.IsBound())
+		{
+			return IsActiveEvent.Execute();
+		}
+
+		return false;
+	}
 
 	// IFilter implementation
 	DECLARE_DERIVED_EVENT( FFilterBase, IFilter<FilterType>::FChangedEvent, FChangedEvent );
@@ -84,6 +95,9 @@ protected:
 	DECLARE_EVENT_OneParam(FFilterBase, FSetActiveEvent, bool);
 	FSetActiveEvent SetActiveEvent;
 
+	DECLARE_DELEGATE_RetVal(bool, FIsActiveEvent);
+	FIsActiveEvent IsActiveEvent;
+	
 	TSharedPtr<FFilterCategory> FilterCategory;
 
 	friend struct FFrontendFilterExternalActivationHelper<FilterType>;
diff --git a/Engine/Source/Developer/ToolWidgets/Public/Filters/SBasicFilterBar.h b/Engine/Source/Developer/ToolWidgets/Public/Filters/SBasicFilterBar.h
--- a/Engine/Source/Developer/ToolWidgets/Public/Filters/SBasicFilterBar.h
+++ b/Engine/Source/Developer/ToolWidgets/Public/Filters/SBasicFilterBar.h
@@ -1239,8 +1239,8 @@ protected:
 		RemoveAllFilters();
 	}
 
- 	/** Called to set a filter active externally */
- 	void OnSetFilterActive(bool bInActive, TWeakPtr<FFilterBase<FilterType>> InWeakFilter)
+ 	/** Called externally to check/pin a filter, and activates/deactivates it */
+ 	void OnSetFilterActive(bool bActive, TWeakPtr<FFilterBase<FilterType>> InWeakFilter)
  	{
  		TSharedPtr<FFilterBase<FilterType>> Filter = InWeakFilter.Pin();
  		if (Filter.IsValid())
@@ -1248,10 +1248,40 @@ protected:
  			if (!IsFrontendFilterInUse(Filter.ToSharedRef()))
  			{
  				TSharedRef<SFilter> NewFilter = AddFilterToBar(Filter.ToSharedRef());
- 				NewFilter->SetEnabled(bInActive);
+ 				NewFilter->SetEnabled(bActive);
+ 			}
+ 			else
+ 			{
+ 				for (const TSharedRef<SFilter>& PinnedFilter : Filters)
+ 				{
+ 					if (PinnedFilter->GetFrontendFilter() == Filter)
+ 					{
+ 						PinnedFilter->SetEnabled(bActive);
+ 						break;
+ 					}
+ 				}
  			}
  		}
  	}
+
+	/** Called externally to determine if a filter has been checked/pinned and activated */
+	bool OnIsFilterActive(TWeakPtr<FFilterBase<FilterType>> InWeakFilter)
+	{
+		TSharedPtr<FFilterBase<FilterType>> Filter = InWeakFilter.Pin();
+		if(Filter.IsValid())
+		{
+			if(!IsFrontendFilterInUse(Filter.ToSharedRef()))
+			{
+				return false;
+			}
+			else
+			{
+				return IsFilterActive(Filter);
+			}
+		}
+
+		return false;
+	}
  	
  	/** Handler for when a checkbox next to a custom text filter is clicked */
 	void CustomTextFilterClicked(ECheckBoxState CheckBoxState, TSharedRef<ICustomTextFilter<FilterType>> Filter)
@@ -1669,6 +1699,7 @@ struct FFrontendFilterExternalActivationHelper
 	{
 		TWeakPtr<FFilterBase<FilterType>> WeakFilter = InFrontendFilter;
 		InFrontendFilter->SetActiveEvent.AddSP(&InFilterList.Get(), &SBasicFilterBar<FilterType>::OnSetFilterActive, WeakFilter);
+		InFrontendFilter->IsActiveEvent.BindSP(&InFilterList.Get(), &SBasicFilterBar<FilterType>::OnIsFilterActive, WeakFilter);
 	}
 };
 
