Adding support for getting and setting the SFilterSearchBox search history text

diff --git a/Engine/Source/Developer/ToolWidgets/Private/Filters/SFilterSearchBox.cpp b/Engine/Source/Developer/ToolWidgets/Private/Filters/SFilterSearchBox.cpp
--- a/Engine/Source/Developer/ToolWidgets/Private/Filters/SFilterSearchBox.cpp
+++ b/Engine/Source/Developer/ToolWidgets/Private/Filters/SFilterSearchBox.cpp
@@ -154,9 +154,21 @@ void SFilterSearchBox::Construct( const FArguments& InArgs )
 	OnTextCommitted = InArgs._OnTextCommitted;
 	bShowSearchHistory = InArgs._ShowSearchHistory;
 
-	// Default text shown when there are no items in the search history
-	EmptySearchHistoryText = MakeShareable(new FText(LOCTEXT("EmptySearchHistoryText", "The Search History is Empty")));
-	SearchHistory.Add(EmptySearchHistoryText);
+	if (InArgs._SearchHistory.IsEmpty())
+	{
+		// Default text shown when there are no items in the search history
+		EmptySearchHistoryText = MakeShareable(new FText(LOCTEXT("EmptySearchHistoryText", "The Search History is Empty")));
+		SearchHistory.Add(EmptySearchHistoryText);
+	}
+	else
+	{
+		// Restore search history if supplied
+		SearchHistory.Reserve(InArgs._SearchHistory.Num());
+		for (const FText& SearchText : InArgs._SearchHistory)
+		{
+			SearchHistory.Add(MakeShared<FText>(SearchText));
+		}
+	}
 	
 	ChildSlot
 	[
@@ -235,6 +247,24 @@ FText SFilterSearchBox::GetText() const
 	return SearchBox->GetText();
 }
 
+TArray<FText> SFilterSearchBox::GetSearchHistory() const
+{
+	TArray<FText> Results;
+
+	// If there is only one item, check that it's not the placeholder empty search text
+	if (SearchHistory.Num() == 1 && SearchHistory[0] == EmptySearchHistoryText)
+	{
+		return Results;
+	}
+
+	Results.Reserve(SearchHistory.Num());
+	for (const TSharedPtr<FText>& SearchHistoryText : SearchHistory)
+	{
+		Results.Add(*SearchHistoryText.Get());
+	}
+	return Results;
+}
+
 void SFilterSearchBox::SetError( const FText& InError )
 {
 	SearchBox->SetError(InError);
diff --git a/Engine/Source/Developer/ToolWidgets/Public/Filters/SFilterSearchBox.h b/Engine/Source/Developer/ToolWidgets/Public/Filters/SFilterSearchBox.h
--- a/Engine/Source/Developer/ToolWidgets/Public/Filters/SFilterSearchBox.h
+++ b/Engine/Source/Developer/ToolWidgets/Public/Filters/SFilterSearchBox.h
@@ -49,6 +49,9 @@ public:
 	/** The maximum number of items to show in the Search History */
 	SLATE_ARGUMENT(int32, MaxSearchHistory)
 
+	/** The Search History to restore */
+	SLATE_ARGUMENT(TArrayView<FText>, SearchHistory)
+
 	/** Whether the SearchBox should delay notifying listeners of text changed events until the user is done typing */
 	SLATE_ATTRIBUTE( bool, DelayChangeNotificationsWhileTyping )
 
@@ -69,6 +72,9 @@ public:
 	/** Get the text string currently being edited */
 	FText GetText() const;
 
+	/** Get the current Search History excluding the placeholder empty search text */
+	TArray<FText> GetSearchHistory() const;
+
 	/** Set or clear the current error reporting information for this search box */
 	void SetError( const FText& InError );
 	void SetError( const FString& InError );
