Add user avatar widget to ToolWidgets

Allows for other modules or plugins to reuse the widget

diff --git a/Engine/Source/Developer/ToolWidgets/Private/SAvatar.cpp b/Engine/Source/Developer/ToolWidgets/Private/SAvatar.cpp
new file mode 100644
--- /dev/null
+++ b/Engine/Source/Developer/ToolWidgets/Private/SAvatar.cpp
@@ -0,0 +1,72 @@
+// Copyright Epic Games, Inc. All Rights Reserved.
+
+#include "SAvatar.h"
+#include "Framework/Application/SlateApplication.h"
+#include "Widgets/Layout/SBox.h"
+#include "Styling/CoreStyle.h"
+#include "Fonts/FontMeasure.h"
+#include "Math/RandomStream.h"
+
+void SAvatar::Construct(const FArguments& InArgs)
+{
+	Identifier = InArgs._Identifier;
+	Description = InArgs._Description;
+
+	bShowInitial = InArgs._ShowInitial;
+
+	const uint32 Hash = GetTypeHash(Identifier);
+	const uint8 H = ((Hash >> 0) & 0xFF) ^ ((Hash >> 8) & 0xFF) ^ ((Hash >> 16) & 0xFF) ^ ((Hash >> 24) & 0xFF);
+	const uint8 S = 128 + ((Hash >> 8) & 0x7F);
+	const uint8 V = 128 + ((Hash >> 16) & 0x7F) / 2;
+	BackgroundColor = FLinearColor::MakeFromHSV8(H, S, V).ToFColor(false);
+	ForegroundColor = FColor::White;
+
+	ChildSlot
+	[
+		SNew(SBox)
+		.MinDesiredWidth(InArgs._WidthOverride)
+		.MinDesiredHeight(InArgs._HeightOverride)
+		.Padding(0.0f)
+	];
+}
+
+int32 SAvatar::OnPaint(const FPaintArgs& Args, const FGeometry& AllottedGeometry, const FSlateRect& MyCullingRect, FSlateWindowElementList& OutDrawElements, int32 LayerId, const FWidgetStyle& InWidgetStyle, bool bParentEnabled) const
+{
+	const float Radius = FMath::Max(AllottedGeometry.GetLocalSize().X, AllottedGeometry.GetLocalSize().Y) * 0.5f;
+
+	// Draw background circle.
+
+	const FSlateBrush* CircleBrush = FCoreStyle::Get().GetBrush("Icons.FilledCircle");
+
+	FSlateDrawElement::MakeBox(
+		OutDrawElements,
+		LayerId++,
+		AllottedGeometry.ToPaintGeometry(),
+		CircleBrush,
+		bParentEnabled && IsEnabled() ? ESlateDrawEffect::None : ESlateDrawEffect::DisabledEffect,
+		BackgroundColor
+	);
+
+	// Draw foreground text.
+
+	if (Description.Len() > 0 && bShowInitial)
+	{
+		FSlateFontInfo FontInfo = FCoreStyle::GetDefaultFontStyle("Bold", Radius);
+
+		const FString Text = Description.Left(1).ToUpper();
+		const FVector2D TextSize = FSlateApplication::Get().GetRenderer()->GetFontMeasureService()->Measure(Text, FontInfo);
+		const FVector2D TextOffset(Radius - TextSize.X / 2, Radius - TextSize.Y / 2);
+
+		FSlateDrawElement::MakeText(
+			OutDrawElements,
+			LayerId++,
+			AllottedGeometry.ToPaintGeometry(AllottedGeometry.GetLocalSize(), FSlateLayoutTransform(TextOffset)),
+			Text,
+			FontInfo,
+			bParentEnabled && IsEnabled() ? ESlateDrawEffect::None : ESlateDrawEffect::DisabledEffect,
+			ForegroundColor
+		);
+	}
+
+	return SCompoundWidget::OnPaint(Args, AllottedGeometry, MyCullingRect, OutDrawElements, LayerId, InWidgetStyle, bParentEnabled && IsEnabled());
+}
\ No newline at end of file
diff --git a/Engine/Source/Developer/ToolWidgets/Public/SAvatar.h b/Engine/Source/Developer/ToolWidgets/Public/SAvatar.h
new file mode 100644
--- /dev/null
+++ b/Engine/Source/Developer/ToolWidgets/Public/SAvatar.h
@@ -0,0 +1,49 @@
+// Copyright Epic Games, Inc. All Rights Reserved.
+
+#pragma once
+
+#include "CoreMinimal.h"
+#include "Widgets/SCompoundWidget.h"
+
+/** Widget for displaying a user avatar icon.
+ * The avatar's background color is computed based on the Identifier.
+ * The avatar's displayed letter is the initial letter of the Description.
+ */
+class TOOLWIDGETS_API SAvatar : public SCompoundWidget
+{
+public:
+	SLATE_BEGIN_ARGS(SAvatar) :
+		_HeightOverride(17.f),
+		_WidthOverride(17.f),
+		_ShowInitial(true)
+	{}
+		SLATE_ARGUMENT(FString, Identifier)
+		SLATE_ARGUMENT(FString, Description)
+		SLATE_ARGUMENT(float, HeightOverride)
+		SLATE_ARGUMENT(float, WidthOverride)
+		SLATE_ARGUMENT(bool, ShowInitial)
+	SLATE_END_ARGS()
+
+public:
+	/** Construct this widget */
+	void Construct(const FArguments& InArgs);
+
+	/** The custom OnPaint handler necessary because we want to draw a circle */
+	virtual int32 OnPaint(const FPaintArgs& Args, const FGeometry& AllottedGeometry, const FSlateRect& MyCullingRect, FSlateWindowElementList& OutDrawElements, int32 LayerId, const FWidgetStyle& InWidgetStyle, bool bParentEnabled) const override;
+
+private:
+	/** The Identifier of the avatar */
+	FString Identifier;
+
+	/** The Description of the avatar */
+	FString Description;
+
+	/** The BackgroundColor used */
+	FColor BackgroundColor;
+
+	/** The ForegroundColor used (text) */
+	FColor ForegroundColor;
+
+	/** The setting that determines if we're showing the initial */
+	bool bShowInitial;
+};
\ No newline at end of file
