Added possbility to specify a color per segment in segmented timeline view

diff --git a/Engine/Source/Developer/ToolWidgets/Private/SSegmentedTimelineView.cpp b/Engine/Source/Developer/ToolWidgets/Private/SSegmentedTimelineView.cpp
--- a/Engine/Source/Developer/ToolWidgets/Private/SSegmentedTimelineView.cpp
+++ b/Engine/Source/Developer/ToolWidgets/Private/SSegmentedTimelineView.cpp
@@ -38,6 +38,9 @@ int32 SSegmentedTimelineView::PaintBlock(const FGeometry& AllottedGeometry, cons
 
 	if (TSharedPtr<FSegmentData> SegmentDataPtr = SegmentData.Get())
 	{
+		uint32 Index = 0;
+		const TOptional<TArray<FLinearColor>>& OverrideColors = SegmentDataPtr->AlternatingSegmentsColors;
+		const bool bOverrideColorsProvided = OverrideColors.IsSet() && OverrideColors.GetValue().Num() > 0;
 		for (const TRange<double>& SegmentRange : SegmentDataPtr->Segments)
 		{
 			double LowerBound = SegmentRange.HasLowerBound() ? SegmentRange.GetLowerBoundValue() : 0;
@@ -48,8 +51,8 @@ int32 SSegmentedTimelineView::PaintBlock(const FGeometry& AllottedGeometry, cons
 			FPaintGeometry BoxGeometry = AllottedGeometry.ToPaintGeometry(FVector2f(BoxMax - BoxMin, Size.Y - 2), FSlateLayoutTransform(FVector2f(BoxMin, 1.f)));
 
 			const FSlateBrush* Brush = FAppStyle::GetBrush("Sequencer.SectionArea.Background");
-
-			FSlateDrawElement::MakeBox(OutDrawElements, LayerId, BoxGeometry, Brush, ESlateDrawEffect::None, Color);
+			FSlateDrawElement::MakeBox(OutDrawElements, LayerId, BoxGeometry, Brush, ESlateDrawEffect::None, bOverrideColorsProvided ? OverrideColors.GetValue()[Index % OverrideColors->Num()] : Color);
+			++Index;
 		}
 		
 		++LayerId;
diff --git a/Engine/Source/Developer/ToolWidgets/Public/SSegmentedTimelineView.h b/Engine/Source/Developer/ToolWidgets/Public/SSegmentedTimelineView.h
--- a/Engine/Source/Developer/ToolWidgets/Public/SSegmentedTimelineView.h
+++ b/Engine/Source/Developer/ToolWidgets/Public/SSegmentedTimelineView.h
@@ -19,6 +19,9 @@ public:
 	{
 	public:
 		TArray<TRange<double>> Segments;
+		// This is an optional array of alternating colors. Every segment will use the next color in this array and will wrap around
+		// if the number of colors provided is > than the number of segments.
+		TOptional<TArray<FLinearColor>> AlternatingSegmentsColors;
 	};
 
 	SLATE_BEGIN_ARGS(SSegmentedTimelineView)
