PR #10485: Make the Content Browser 'Filter' toggle buttons hit-test the entire button rather than just the text

diff --git a/Engine/Source/Developer/ToolWidgets/Public/Filters/SBasicFilterBar.h b/Engine/Source/Developer/ToolWidgets/Public/Filters/SBasicFilterBar.h
--- a/Engine/Source/Developer/ToolWidgets/Public/Filters/SBasicFilterBar.h
+++ b/Engine/Source/Developer/ToolWidgets/Public/Filters/SBasicFilterBar.h
@@ -501,37 +501,37 @@ protected:
 				break;
 			case EFilterPillStyle::Default:
 			default:
-				ContentWidget = SNew(SBorder)
-				 .Padding(1.0f)
-				 .BorderImage(FAppStyle::Get().GetBrush("FilterBar.FilterBackground"))
-				 [
-					 SNew(SHorizontalBox)
-					 +SHorizontalBox::Slot()
-					 .VAlign(VAlign_Center)
-					 .AutoWidth()
-					 [
-						 SNew(SImage)
-						 .Image(FAppStyle::Get().GetBrush("FilterBar.FilterImage"))
-						 .ColorAndOpacity(this, &SFilter::GetFilterImageColorAndOpacity)
-					 ]
-					 +SHorizontalBox::Slot()
-					 .Padding(TAttribute<FMargin>(this, &SFilter::GetFilterNamePadding))
-					 .VAlign(VAlign_Center)
-					 [
-						 SAssignNew( ToggleButtonPtr, SFilterCheckBox )
-						.Style(FAppStyle::Get(), "FilterBar.FilterButton")
-						.ToolTipText(FilterToolTip)
-						.IsChecked(this, &SFilter::IsChecked)
-						.OnCheckStateChanged(this, &SFilter::FilterToggled)
-						.CheckBoxContentUsesAutoWidth(false)
-						.OnGetMenuContent(this, &SFilter::GetRightClickMenuContent)
+				ContentWidget = SAssignNew( ToggleButtonPtr, SFilterCheckBox )
+					.Style(FAppStyle::Get(), "FilterBar.FilterButton")
+					.ToolTipText(FilterToolTip)
+					.IsChecked(this, &SFilter::IsChecked)
+					.OnCheckStateChanged(this, &SFilter::FilterToggled)
+					.CheckBoxContentUsesAutoWidth(false)
+					.OnGetMenuContent(this, &SFilter::GetRightClickMenuContent)
+					[
+						SNew(SBorder)
+						.Padding(1.0f)
+						.BorderImage(FAppStyle::Get().GetBrush("FilterBar.FilterBackground"))
 						[
-							SNew(STextBlock)
-							.Text(this, &SFilter::GetFilterDisplayName)
-							.IsEnabled_Lambda([this] {return bEnabled;})
+							SNew(SHorizontalBox)
+							+SHorizontalBox::Slot()
+							.VAlign(VAlign_Center)
+							.AutoWidth()
+							[
+								SNew(SImage)
+								.Image(FAppStyle::Get().GetBrush("FilterBar.FilterImage"))
+								.ColorAndOpacity(this, &SFilter::GetFilterImageColorAndOpacity)
+							]
+							+SHorizontalBox::Slot()
+							.Padding(TAttribute<FMargin>(this, &SFilter::GetFilterNamePadding))
+							.VAlign(VAlign_Center)
+							[
+								SNew(STextBlock)
+								.Text(this, &SFilter::GetFilterDisplayName)
+								.IsEnabled_Lambda([this] {return bEnabled;})
+							]
 						]
-					 ]
-				 ];
+					];
 					
 			}
 			
