[SSimpleTimeSlider] Prevent selection to scrub outside of clamp range

diff --git a/Engine/Source/Developer/ToolWidgets/Private/SSimpleTimeSlider.cpp b/Engine/Source/Developer/ToolWidgets/Private/SSimpleTimeSlider.cpp
--- a/Engine/Source/Developer/ToolWidgets/Private/SSimpleTimeSlider.cpp
+++ b/Engine/Source/Developer/ToolWidgets/Private/SSimpleTimeSlider.cpp
@@ -317,10 +317,7 @@ FReply SSimpleTimeSlider::OnPreviewMouseButtonDown(const FGeometry& MyGeometry,
 	if ( bHandleLeftMouseButton )
 	{
 		// Always capture mouse if we left or right click on the widget
-		FScrubRangeToScreen RangeToScreen(ViewRange.Get(), MyGeometry.GetLocalSize());
-		FVector2D CursorPos = MyGeometry.AbsoluteToLocal(MouseEvent.GetLastScreenSpacePosition());
-		float NewValue = RangeToScreen.LocalXToInput(CursorPos.X);
-
+		const float NewValue = GetTimeAtCursorPosition(MyGeometry, MouseEvent);
 		CommitScrubPosition(NewValue, /*bIsScrubbing=*/false);
 		return FReply::Handled().CaptureMouse(AsShared()).PreventThrottling();
 	}
@@ -352,10 +349,7 @@ FReply SSimpleTimeSlider::OnMouseButtonUp(const FGeometry& MyGeometry, const FPo
 	{
 		if( !bDraggingScrubber )
 		{
-			FScrubRangeToScreen RangeToScreen( ViewRange.Get(), MyGeometry.GetLocalSize() );
-			FVector2D CursorPos = MyGeometry.AbsoluteToLocal(MouseEvent.GetLastScreenSpacePosition());
-			float NewValue = RangeToScreen.LocalXToInput(CursorPos.X);
-
+			const float NewValue = GetTimeAtCursorPosition(MyGeometry, MouseEvent);
 			CommitScrubPosition( NewValue, /*bIsScrubbing=*/false );
 		}
 
