Add AutoWrapText argument to SWarningOrErrorBox

diff --git a/Engine/Source/Developer/ToolWidgets/Private/SWarningOrErrorBox.cpp b/Engine/Source/Developer/ToolWidgets/Private/SWarningOrErrorBox.cpp
--- a/Engine/Source/Developer/ToolWidgets/Private/SWarningOrErrorBox.cpp
+++ b/Engine/Source/Developer/ToolWidgets/Private/SWarningOrErrorBox.cpp
@@ -32,7 +32,7 @@ void SWarningOrErrorBox::Construct(const FArguments& InArgs)
 				SNew(STextBlock)
 				.Text(InArgs._Message)
 				.ColorAndOpacity(FAppStyle::Get().GetSlateColor("Colors.White"))
-				.AutoWrapText(true)
+				.AutoWrapText(InArgs._AutoWrapText)
 			]
 
 			+SHorizontalBox::Slot()
diff --git a/Engine/Source/Developer/ToolWidgets/Public/SWarningOrErrorBox.h b/Engine/Source/Developer/ToolWidgets/Public/SWarningOrErrorBox.h
--- a/Engine/Source/Developer/ToolWidgets/Public/SWarningOrErrorBox.h
+++ b/Engine/Source/Developer/ToolWidgets/Public/SWarningOrErrorBox.h
@@ -18,12 +18,14 @@ public:
 		: _MessageStyle(EMessageStyle::Warning)
 		, _Padding(16.0f)
 		, _IconSize(24,24)
+		, _AutoWrapText(true)
 		, _Content()
 	{}
 		SLATE_ATTRIBUTE(FText, Message)
 		SLATE_ATTRIBUTE(EMessageStyle, MessageStyle)
 		SLATE_ARGUMENT(FMargin, Padding)
 		SLATE_ARGUMENT(FVector2D, IconSize)
+		SLATE_ARGUMENT(bool, AutoWrapText)
 
 		SLATE_DEFAULT_SLOT( FArguments, Content )
 
