Deprecate description field and change hash function

Having two values for an avatar may be not needed, the id might not be available. Moreover, instead of using UE's hash function the new suggested function is simple and portable to use outside UE. This is a best effort to have consistency across UCS products.

diff --git a/Engine/Source/Developer/ToolWidgets/Private/SAvatar.cpp b/Engine/Source/Developer/ToolWidgets/Private/SAvatar.cpp
--- a/Engine/Source/Developer/ToolWidgets/Private/SAvatar.cpp
+++ b/Engine/Source/Developer/ToolWidgets/Private/SAvatar.cpp
@@ -14,10 +14,10 @@ void SAvatar::Construct(const FArguments& InArgs)
 
 	bShowInitial = InArgs._ShowInitial;
 
-	const uint32 Hash = GetTypeHash(Identifier);
-	const uint8 H = ((Hash >> 0) & 0xFF) ^ ((Hash >> 8) & 0xFF) ^ ((Hash >> 16) & 0xFF) ^ ((Hash >> 24) & 0xFF);
-	const uint8 S = 128 + ((Hash >> 8) & 0x7F);
-	const uint8 V = 128 + ((Hash >> 16) & 0x7F) / 2;
+	const uint64 Value = Hash();
+	const uint8 H = ((Value >> 0) & 0xFF) ^ ((Value >> 8) & 0xFF) ^ ((Value >> 16) & 0xFF) ^ ((Value >> 24) & 0xFF);
+	const uint8 S = 128 + ((Value >> 8) & 0x7F);
+	const uint8 V = 128 + ((Value >> 16) & 0x7F) / 2;
 	BackgroundColor = FLinearColor::MakeFromHSV8(H, S, V).ToFColor(false);
 	ForegroundColor = FColor::White;
 
@@ -49,11 +49,11 @@ int32 SAvatar::OnPaint(const FPaintArgs& Args, const FGeometry& AllottedGeometry
 
 	// Draw foreground text.
 
-	if (Description.Len() > 0 && bShowInitial)
+	if (Identifier.Len() > 0 && bShowInitial)
 	{
 		FSlateFontInfo FontInfo = FCoreStyle::GetDefaultFontStyle("Bold", Radius);
 
-		const FString Text = Description.Left(1).ToUpper();
+		const FString Text = Identifier.Left(1).ToUpper();
 		const FVector2D TextSize = FSlateApplication::Get().GetRenderer()->GetFontMeasureService()->Measure(Text, FontInfo);
 		const FVector2D TextOffset(Radius - TextSize.X / 2, Radius - TextSize.Y / 2);
 
@@ -69,4 +69,15 @@ int32 SAvatar::OnPaint(const FPaintArgs& Args, const FGeometry& AllottedGeometry
 	}
 
 	return SCompoundWidget::OnPaint(Args, AllottedGeometry, MyCullingRect, OutDrawElements, LayerId, InWidgetStyle, bParentEnabled && IsEnabled());
+}
+
+uint64 SAvatar::Hash() const
+{
+	uint64 HashValue = 5381;
+	for (char Char : Identifier.ToLower())
+	{
+		HashValue = ((HashValue << 5) + HashValue) + Char;
+	}
+
+	return HashValue;
 }
\ No newline at end of file
diff --git a/Engine/Source/Developer/ToolWidgets/Public/SAvatar.h b/Engine/Source/Developer/ToolWidgets/Public/SAvatar.h
--- a/Engine/Source/Developer/ToolWidgets/Public/SAvatar.h
+++ b/Engine/Source/Developer/ToolWidgets/Public/SAvatar.h
@@ -18,7 +18,7 @@ public:
 		_ShowInitial(true)
 	{}
 		SLATE_ARGUMENT(FString, Identifier)
-		SLATE_ARGUMENT(FString, Description)
+		SLATE_ARGUMENT_DEPRECATED(FString, Description, 5.5, "Description is deprecated, use Identifier instead.")
 		SLATE_ARGUMENT(float, HeightOverride)
 		SLATE_ARGUMENT(float, WidthOverride)
 		SLATE_ARGUMENT(bool, ShowInitial)
@@ -32,6 +32,9 @@ public:
 	virtual int32 OnPaint(const FPaintArgs& Args, const FGeometry& AllottedGeometry, const FSlateRect& MyCullingRect, FSlateWindowElementList& OutDrawElements, int32 LayerId, const FWidgetStyle& InWidgetStyle, bool bParentEnabled) const override;
 
 private:
+	/** Computes the hash of the identifier with djb2 function*/
+	uint64 Hash() const;
+
 	/** The Identifier of the avatar */
 	FString Identifier;
 
