Add RemoveAll method to SCheckBoxList.

diff --git a/Engine/Source/Developer/ToolWidgets/Private/SCheckBoxList.cpp b/Engine/Source/Developer/ToolWidgets/Private/SCheckBoxList.cpp
--- a/Engine/Source/Developer/ToolWidgets/Private/SCheckBoxList.cpp
+++ b/Engine/Source/Developer/ToolWidgets/Private/SCheckBoxList.cpp
@@ -142,6 +142,12 @@ int32 SCheckBoxList::AddItem(TSharedRef<SWidget> Widget, bool bIsChecked)
 	return ReturnValue;
 }
 
+void SCheckBoxList::RemoveAll()
+{
+	Items.Reset();
+	ListView->RebuildList();
+}
+
 void SCheckBoxList::RemoveItem(int32 Index)
 {
 	if (Items.IsValidIndex(Index))
diff --git a/Engine/Source/Developer/ToolWidgets/Public/SCheckBoxList.h b/Engine/Source/Developer/ToolWidgets/Public/SCheckBoxList.h
--- a/Engine/Source/Developer/ToolWidgets/Public/SCheckBoxList.h
+++ b/Engine/Source/Developer/ToolWidgets/Public/SCheckBoxList.h
@@ -48,6 +48,7 @@ public:
 
 	int32 AddItem(const FText& Text, bool bIsChecked);
 	int32 AddItem(TSharedRef<SWidget> Widget, bool bIsChecked);
+	void RemoveAll();
 	void RemoveItem(int32 Index);
 	bool IsItemChecked(int32 Index) const;
 
