Fixed text alignment on SPrimaryButton when the button's icon was set to nullptr, or to a delegate that returned nullptr.

diff --git a/Engine/Source/Developer/ToolWidgets/Private/SPrimaryButton.cpp b/Engine/Source/Developer/ToolWidgets/Private/SPrimaryButton.cpp
--- a/Engine/Source/Developer/ToolWidgets/Private/SPrimaryButton.cpp
+++ b/Engine/Source/Developer/ToolWidgets/Private/SPrimaryButton.cpp
@@ -7,22 +7,42 @@
 
 void SPrimaryButton::Construct(const FArguments& InArgs)
 {
+	// To know if the icon is set, check explicitely for a brush or delegate, because
+	// if a caller passed nullptr, that would result in _Icon.IsSet() returning true.
+	// (Specifically checking IsBound() first so that if it is, we don't invoke it by calling Get())
+	const bool bIconSet = InArgs._Icon.IsBound() || InArgs._Icon.Get();
+
+	// If a delegate is bound, we'll store it and pass our own delegate to the Image.
+	// This will allow us to update the layout as an icon comes and goes.
+	TAttribute<const FSlateBrush*> IconAttribute;
+	if (InArgs._Icon.IsBound())
+	{
+		OwnerDelegate = InArgs._Icon.GetBinding();
+
+		TAttribute<const FSlateBrush*>::FGetter Delegate = TAttribute<const FSlateBrush*>::FGetter::CreateSP(this, &SPrimaryButton::OnGetBrush);		
+		IconAttribute.Bind(MoveTemp(Delegate));
+	}
+	else
+	{
+		IconAttribute = InArgs._Icon;
+	}
+
 	SButton::Construct(SButton::FArguments()
-		.ButtonStyle(&FAppStyle::Get().GetWidgetStyle<FButtonStyle>(InArgs._Icon.IsSet() ? "PrimaryButtonLabelAndIcon" : "PrimaryButton"))
+		.ButtonStyle(&FAppStyle::Get().GetWidgetStyle<FButtonStyle>(bIconSet ? "PrimaryButtonLabelAndIcon" : "PrimaryButton"))
 		.OnClicked(InArgs._OnClicked)
 		.ForegroundColor(FSlateColor::UseStyle())
 		.HAlign(HAlign_Center)
 		[
-			SNew(SHorizontalBox)
+			SAssignNew(HorizontalBox, SHorizontalBox)
 			+SHorizontalBox::Slot()
 			.AutoWidth()
 			.VAlign(VAlign_Center)
-			.Padding(FMargin(0,0,3,0))
+			.Padding(FMargin(0,0,bIconSet ? 3 : 0,0))
 			[
 				SNew(SImage)
 				.ColorAndOpacity(FSlateColor::UseForeground())
-				.Image(InArgs._Icon)
-				.Visibility(InArgs._Icon.IsSet() ? EVisibility::HitTestInvisible : EVisibility::Collapsed)
+				.Image(IconAttribute)
+				.Visibility(bIconSet ? EVisibility::HitTestInvisible : EVisibility::Collapsed)
 			]
 			+ SHorizontalBox::Slot()
 			.VAlign(VAlign_Center)
@@ -37,3 +57,17 @@ void SPrimaryButton::Construct(const FArguments& InArgs)
 	);
 }
 
+const FSlateBrush* SPrimaryButton::OnGetBrush()
+{
+	const FSlateBrush* Brush = nullptr;
+	if (OwnerDelegate.IsBound())
+	{
+		Brush = OwnerDelegate.Execute();
+	}
+
+	// Update the button's styling based on whether a brush will be rendered or not
+	SetButtonStyle(&FAppStyle::Get().GetWidgetStyle<FButtonStyle>(Brush ? "PrimaryButtonLabelAndIcon" : "PrimaryButton"));
+	HorizontalBox->GetSlot(0).SetPadding(FMargin(0,0,Brush ? 3 : 0,0));
+
+	return Brush;
+}
\ No newline at end of file
diff --git a/Engine/Source/Developer/ToolWidgets/Public/SPrimaryButton.h b/Engine/Source/Developer/ToolWidgets/Public/SPrimaryButton.h
--- a/Engine/Source/Developer/ToolWidgets/Public/SPrimaryButton.h
+++ b/Engine/Source/Developer/ToolWidgets/Public/SPrimaryButton.h
@@ -6,6 +6,8 @@
 #include "Widgets/DeclarativeSyntaxSupport.h"
 #include "Widgets/Input/SButton.h"
 
+class SHorizontalBox;
+
 class TOOLWIDGETS_API SPrimaryButton : public SButton
 {
 public:
@@ -26,4 +28,11 @@ public:
 	SLATE_END_ARGS()
 
 	void Construct(const FArguments& InArgs);
+
+private:
+	const FSlateBrush* OnGetBrush();
+
+private:
+	TAttribute<const FSlateBrush*>::FGetter OwnerDelegate;
+	TSharedPtr<SHorizontalBox> HorizontalBox;
 };
\ No newline at end of file
