Allow SMessageDialog to optionally use plain-text

This is controlled via a new UseRichText attribute, and may be set to false when rich-text support is not required.

diff --git a/Engine/Source/Developer/ToolWidgets/Private/Dialog/SMessageDialog.cpp b/Engine/Source/Developer/ToolWidgets/Private/Dialog/SMessageDialog.cpp
--- a/Engine/Source/Developer/ToolWidgets/Private/Dialog/SMessageDialog.cpp
+++ b/Engine/Source/Developer/ToolWidgets/Private/Dialog/SMessageDialog.cpp
@@ -7,6 +7,7 @@
 #include "Styling/AppStyle.h"
 #include "Styling/SlateBrush.h"
 #include "Widgets/Text/SRichTextBlock.h"
+#include "Widgets/Text/STextBlock.h"
 #include "Widgets/Input/SButton.h"
 #include "Widgets/Images/SImage.h"
 
@@ -14,17 +15,35 @@ void SMessageDialog::Construct(const FArguments& InArgs)
 {
 	Message = InArgs._Message;
 	
-	TSharedPtr<SRichTextBlock> RichTextBlock;
+	TSharedPtr<SWidget> TextBlockWidget;
+	if (InArgs._UseRichText)
+	{
+		TSharedRef<SRichTextBlock> RichTextBlock = SNew(SRichTextBlock)
+			.Text(Message)
+			.MinDesiredWidth(InArgs._ContentMinWidth)
+			.WrapTextAt(InArgs._WrapMessageAt)
+			.Decorators(InArgs._Decorators);
+
+		if (InArgs._DecoratorStyleSet)
+		{
+			RichTextBlock->SetDecoratorStyleSet(InArgs._DecoratorStyleSet);
+		}
+
+		TextBlockWidget = RichTextBlock;
+	}
+	else
+	{
+		TextBlockWidget = SNew(STextBlock)
+			.Text(Message)
+			.MinDesiredWidth(InArgs._ContentMinWidth)
+			.WrapTextAt(InArgs._WrapMessageAt);
+	}
 
 	SCustomDialog::Construct(SCustomDialog::FArguments()
 		.Title(InArgs._Title)
 		.Content()
 		[
-			SAssignNew(RichTextBlock, SRichTextBlock)
-			.MinDesiredWidth(InArgs._ContentMinWidth)
-			.Text(Message)
-			.WrapTextAt(InArgs._WrapMessageAt)
-			.Decorators(InArgs._Decorators)
+			TextBlockWidget.ToSharedRef()
 		]
 		.WindowArguments(InArgs._WindowArguments)
 		.RootPadding(16.f)
@@ -56,11 +75,6 @@ void SMessageDialog::Construct(const FArguments& InArgs)
 			]
 		]
 	);
-
-	if (InArgs._DecoratorStyleSet)
-	{
-		RichTextBlock->SetDecoratorStyleSet(InArgs._DecoratorStyleSet);
-	}
 }
 
 FReply SMessageDialog::OnKeyDown(const FGeometry& MyGeometry, const FKeyEvent& InKeyEvent)
diff --git a/Engine/Source/Developer/ToolWidgets/Public/Dialog/SMessageDialog.h b/Engine/Source/Developer/ToolWidgets/Public/Dialog/SMessageDialog.h
--- a/Engine/Source/Developer/ToolWidgets/Public/Dialog/SMessageDialog.h
+++ b/Engine/Source/Developer/ToolWidgets/Public/Dialog/SMessageDialog.h
@@ -17,6 +17,7 @@ public:
 
 	SLATE_BEGIN_ARGS(SMessageDialog)
 		: _AutoCloseOnButtonPress(true)
+		, _UseRichText(true)
 		, _DecoratorStyleSet(nullptr)
 		, _Icon(nullptr)
 		, _UseScrollBox(true)
@@ -44,10 +45,13 @@ public:
 		/** Whether to automatically close this window when any button is pressed (default: true) */
 		SLATE_ARGUMENT(bool, AutoCloseOnButtonPress)
 	
-		/** Text decorators used while parsing the rich text messages */
+		/** Whether to use rich-text (true) or plain-text (false) (default: true) */
+		SLATE_ARGUMENT(bool, UseRichText)
+
+		/** Text decorators used while parsing the rich text messages (requires UseRichText: true) */
 		SLATE_ARGUMENT(TArray<TSharedRef<class ITextDecorator>>, Decorators)
 
-		/** Style set used to look up styles used by decorators for rich text messages */
+		/** Style set used to look up styles used by decorators for rich text messages (requires UseRichText: true) */
 		SLATE_ARGUMENT(const ISlateStyle*, DecoratorStyleSet)
 		
 		/********** Cosmetic **********/
