SFilterSearchBox: Add a max size for the search history and make the text clip properly using ellipses

diff --git a/Engine/Source/Developer/ToolWidgets/Private/Filters/SFilterSearchBox.cpp b/Engine/Source/Developer/ToolWidgets/Private/Filters/SFilterSearchBox.cpp
--- a/Engine/Source/Developer/ToolWidgets/Private/Filters/SFilterSearchBox.cpp
+++ b/Engine/Source/Developer/ToolWidgets/Private/Filters/SFilterSearchBox.cpp
@@ -304,10 +304,12 @@ TSharedRef<ITableRow> SFilterSearchBox::MakeSearchHistoryRowWidget(TSharedPtr<FT
 	RowWidget->AddSlot()
 	.HAlign(TextAlignment)
 	.VAlign(VAlign_Center)
+	.MaxWidth(500.0f)
 	.FillWidth(1.0)
 	[
 		SNew(STextBlock)
 		.Text(*SearchText.Get())
+		.OverflowPolicy(ETextOverflowPolicy::Ellipsis)
 	];
 	
 	return SNew(STableRow< TSharedPtr<FString> >, OwnerTable)
