UE-146290 - [MetaSounds] - Dropdown menus do not support selection with arrow keys

diff --git a/Engine/Plugins/Runtime/Metasound/Source/MetasoundEditor/Private/MetasoundNodeDetailCustomization.cpp b/Engine/Plugins/Runtime/Metasound/Source/MetasoundEditor/Private/MetasoundNodeDetailCustomization.cpp
--- a/Engine/Plugins/Runtime/Metasound/Source/MetasoundEditor/Private/MetasoundNodeDetailCustomization.cpp
+++ b/Engine/Plugins/Runtime/Metasound/Source/MetasoundEditor/Private/MetasoundNodeDetailCustomization.cpp
@@ -990,6 +990,7 @@ PRAGMA_ENABLE_DEPRECATION_WARNINGS
 							OnDataTypeSelected(FName(*InNewName));
 						}
 					})
+					.bAlwaysSelectItem(true)
 					.Content()
 					[
 						SNew(STextBlock)
diff --git a/Engine/Source/Developer/ToolWidgets/Private/SSearchableComboBox.cpp b/Engine/Source/Developer/ToolWidgets/Private/SSearchableComboBox.cpp
--- a/Engine/Source/Developer/ToolWidgets/Private/SSearchableComboBox.cpp
+++ b/Engine/Source/Developer/ToolWidgets/Private/SSearchableComboBox.cpp
@@ -18,6 +18,7 @@ void SSearchableComboBox::Construct(const FArguments& InArgs)
 
 	this->OnComboBoxOpening = InArgs._OnComboBoxOpening;
 	this->OnSelectionChanged = InArgs._OnSelectionChanged;
+	this->bAlwaysSelectItem = InArgs._bAlwaysSelectItem;
 	this->OnGenerateWidget = InArgs._OnGenerateWidget;
 
 	OptionsSource = InArgs._OptionsSource;
@@ -223,7 +224,7 @@ void SSearchableComboBox::OnSelectionChanged_Internal(TSharedPtr<FString> Propos
 	}
 
 	// Ensure that the proposed selection is different from selected
-	if (ProposedSelection != SelectedItem)
+	if (ProposedSelection != SelectedItem || bAlwaysSelectItem)
 	{
 		SelectedItem = ProposedSelection;
 		OnSelectionChanged.ExecuteIfBound(ProposedSelection, SelectInfo);
diff --git a/Engine/Source/Developer/ToolWidgets/Public/SSearchableComboBox.h b/Engine/Source/Developer/ToolWidgets/Public/SSearchableComboBox.h
--- a/Engine/Source/Developer/ToolWidgets/Public/SSearchableComboBox.h
+++ b/Engine/Source/Developer/ToolWidgets/Public/SSearchableComboBox.h
@@ -31,6 +31,7 @@ public:
 		, _ForegroundColor(FSlateColor::UseStyle())
 		, _OptionsSource()
 		, _OnSelectionChanged()
+		, _bAlwaysSelectItem(false)
 		, _OnGenerateWidget()
 		, _InitiallySelectedItem(nullptr)
 		, _Method()
@@ -54,6 +55,7 @@ public:
 
 		SLATE_ARGUMENT(const TArray< TSharedPtr<FString> >*, OptionsSource)
 		SLATE_EVENT(FOnSelectionChanged, OnSelectionChanged)
+		SLATE_ARGUMENT(bool, bAlwaysSelectItem)
 		SLATE_EVENT(FOnGenerateWidget, OnGenerateWidget)
 
 		/** Called when combo box is opened, before list is actually created */
@@ -164,6 +166,9 @@ private:
 
 	/** Filtered list that is actually displayed */
 	TArray< TSharedPtr<FString> > FilteredOptionsSource;
+
+	/** If true OnSelectionChanged delegate is executed on every input. */
+	bool bAlwaysSelectItem : 1;
 };
 
 
