[ToolWidgets] Added SActionButton, and exposed the icon color. This is a less specialized version of SPositive/NegativeActionButton that doesn't imply intent.

diff --git a/Engine/Source/Developer/ToolWidgets/Private/SActionButton.cpp b/Engine/Source/Developer/ToolWidgets/Private/SActionButton.cpp
new file mode 100644
--- /dev/null
+++ b/Engine/Source/Developer/ToolWidgets/Private/SActionButton.cpp
@@ -0,0 +1,228 @@
+// Copyright Epic Games, Inc. All Rights Reserved.
+
+#include "SActionButton.h"
+
+#include "Styling/StyleColors.h"
+#include "ToolWidgetsStyle.h"
+#include "Widgets/Images/SImage.h"
+#include "Widgets/Input/SButton.h"
+#include "Widgets/Layout/SSpacer.h"
+#include "Widgets/Layout/SWidgetSwitcher.h"
+
+namespace UE::ToolWidgets::Private
+{
+	inline FName ActionButtonTypeToStyleName(const EActionButtonType ActionButtonType)
+	{
+		static TMap<EActionButtonType, const FName> Lookup = {
+			{ EActionButtonType::Default, TEXT("ActionButton") },
+			{ EActionButtonType::Simple, TEXT("SimpleComboButton") },
+			{ EActionButtonType::Positive, TEXT("PositiveActionButton") },
+			{ EActionButtonType::Warning, TEXT("NegativeActionButton") },
+			{ EActionButtonType::Error, TEXT("NegativeActionButton") },
+		};
+
+		if (const FName* FoundName = Lookup.Find(ActionButtonType))
+		{
+			return *FoundName;
+		}
+
+		return NAME_None;
+	}
+}
+
+void SActionButton::Construct(const FArguments& InArgs)
+{
+	// Get overridden or default
+	ActionButtonType = InArgs._ActionButtonType.Get(EActionButtonType::Default);
+
+	// If style not explicitly set, derive from ActionButtonType
+	if (!InArgs._ActionButtonStyle)
+	{
+		const FName StyleName = UE::ToolWidgets::Private::ActionButtonTypeToStyleName(ActionButtonType.Get());
+		check(StyleName != NAME_None);
+
+		ActionButtonStyle = &UE::ToolWidgets::FToolWidgetsStyle::Get().GetWidgetStyle<FActionButtonStyle>(StyleName);
+	}
+	else
+	{
+		ActionButtonStyle = InArgs._ActionButtonStyle;
+
+		// ActionButtonStyle was specified, but ActionButtonType was not, so derive one from the other
+		if (!InArgs._ActionButtonType.IsSet())
+		{
+			ActionButtonType = ActionButtonStyle->GetActionButtonType();
+		}
+	}
+
+	ButtonStyle = InArgs._ButtonStyle ? InArgs._ButtonStyle : &ActionButtonStyle->ButtonStyle;
+	IconButtonStyle = InArgs._IconButtonStyle ? InArgs._IconButtonStyle : &ActionButtonStyle->GetIconButtonStyle();
+	ComboButtonStyle = InArgs._ComboButtonStyle ? InArgs._ComboButtonStyle : &ActionButtonStyle->ComboButtonStyle;
+	TextBlockStyle = InArgs._TextBlockStyle ? InArgs._TextBlockStyle : &ActionButtonStyle->TextBlockStyle;
+
+	const EHorizontalAlignment HorizontalContentAlignment = InArgs._HorizontalContentAlignment.IsSet()
+		? InArgs._HorizontalContentAlignment.Get(EHorizontalAlignment::HAlign_Center)
+		: static_cast<EHorizontalAlignment>(ActionButtonStyle->HorizontalContentAlignment);
+
+	const bool bHasDownArrow = InArgs._HasDownArrow.IsSet() ? InArgs._HasDownArrow.Get(false) : ActionButtonStyle->bHasDownArrow;
+
+	// Check for widget level override, then style override, otherwise unset
+	const TAttribute<const FSlateBrush*> Icon = InArgs._Icon.IsSet()
+		? InArgs._Icon.Get()
+		: ActionButtonStyle->IconBrush.IsSet()
+		? &ActionButtonStyle->IconBrush.GetValue()
+		: nullptr;
+
+	const bool bHasIcon = Icon.Get() || Icon.IsBound();
+
+	// Check for widget level override, then style override, otherwise get from ActionButtonType
+	const TAttribute<FSlateColor> IconColorAndOpacity = InArgs._IconColorAndOpacity.IsSet()
+		? InArgs._IconColorAndOpacity.Get()
+		: ActionButtonStyle->IconColorAndOpacity.IsSet()
+		? TAttribute<FSlateColor>(ActionButtonStyle->IconColorAndOpacity.GetValue())
+		: TAttribute<FSlateColor>::Create(TAttribute<FSlateColor>::FGetter::CreateSP(this, &SActionButton::GetIconColorAndOpacity));
+
+	TAttribute<FText> Text = InArgs._Text;
+
+	static constexpr float DefaultIconHeight = UE::ToolWidgets::FToolWidgetsStyle::FActionButton::DefaultIconHeight;
+	static constexpr float IconTextPadding = UE::ToolWidgets::FToolWidgetsStyle::FActionButton::DefaultIconLabelSpacing;
+
+	const TSharedRef<SHorizontalBox> ButtonContentContainer = SNew(SHorizontalBox)
+		+ SHorizontalBox::Slot()
+		.HAlign(HAlign_Center)
+		.VAlign(VAlign_Center)
+		.AutoWidth()
+		.Padding(0)
+		[
+			SNew(SWidgetSwitcher)
+			.WidgetIndex(bHasIcon ? 1 : 0)
+
+			+ SWidgetSwitcher::Slot()
+			[
+				SNew(SSpacer)
+				.Size(FVector2D{ 0, DefaultIconHeight })
+			]
+
+			+ SWidgetSwitcher::Slot()
+			[
+				SNew(SImage)
+				.Image(Icon)
+				.ColorAndOpacity(IconColorAndOpacity)
+				.Visibility(bHasIcon ? EVisibility::HitTestInvisible : EVisibility::Collapsed)
+			]
+		]
+
+		+ SHorizontalBox::Slot()
+		.VAlign(VAlign_Center)
+		.FillWidth(1.0f)
+		.Padding(FMargin(bHasIcon ? IconTextPadding : 0, 0, 0, 0))
+		[
+			SNew(STextBlock)
+			.TextStyle(TextBlockStyle)
+			.Text(InArgs._Text)
+			.Visibility_Lambda([Text]()
+			{
+				return Text.Get(FText::GetEmpty()).IsEmpty()
+					? EVisibility::Collapsed
+					: EVisibility::Visible;
+			})
+		];
+
+	// Treated as a ComboButton if OnClicked is not bound
+	const bool bIsComboButton = !InArgs._OnClicked.IsBound();
+
+	if (bIsComboButton)
+	{
+		static const FMargin DefaultComboButtonContentPadding = FMargin(
+			UE::ToolWidgets::FToolWidgetsStyle::FActionButton::DefaultHorizontalPadding,
+			UE::ToolWidgets::FToolWidgetsStyle::FActionButton::DefaultVerticalPadding);
+
+		const TAttribute<FMargin> ComboButtonContentPadding = InArgs._ButtonContentPadding.IsSet()
+			? InArgs._ButtonContentPadding.Get(DefaultComboButtonContentPadding)
+			: ActionButtonStyle->ComboButtonContentPadding.IsSet()
+			? ActionButtonStyle->ComboButtonContentPadding.GetValue()
+			: ActionButtonStyle->ComboButtonStyle.ContentPadding;
+
+		ChildSlot
+		[
+			SAssignNew(ComboButton, SComboButton)
+			.HasDownArrow(bHasDownArrow)
+			.ContentPadding(ComboButtonContentPadding)
+			.ButtonStyle(bHasIcon ? IconButtonStyle : ButtonStyle)
+			.ComboButtonStyle(ComboButtonStyle)
+			.IsEnabled(InArgs._IsEnabled)
+			.ToolTipText(InArgs._ToolTipText)
+			.HAlign(HorizontalContentAlignment)
+			.VAlign(VAlign_Center)
+			.ButtonContent()
+			[
+				ButtonContentContainer
+			]
+			.MenuContent()
+			[
+				InArgs._MenuContent.Widget
+			]
+			.OnGetMenuContent(InArgs._OnGetMenuContent)
+			.OnMenuOpenChanged(InArgs._OnMenuOpenChanged)
+			.OnComboBoxOpened(InArgs._OnComboBoxOpened)
+		];
+	}
+	else
+	{
+		static const FMargin DefaultButtonContentPadding = FMargin(
+			UE::ToolWidgets::FToolWidgetsStyle::FActionButton::DefaultHorizontalPadding,
+			UE::ToolWidgets::FToolWidgetsStyle::FActionButton::DefaultVerticalPadding);
+
+		const TAttribute<FMargin> ButtonContentPadding = InArgs._ButtonContentPadding.IsSet()
+			? InArgs._ButtonContentPadding.Get(DefaultButtonContentPadding)
+			: ActionButtonStyle->ButtonContentPadding.IsSet()
+			? ActionButtonStyle->ButtonContentPadding.GetValue()
+			: DefaultButtonContentPadding;
+
+		ChildSlot
+		[
+			SAssignNew(Button, SButton)
+			.ContentPadding(ButtonContentPadding)
+			.ButtonStyle(bHasIcon ? IconButtonStyle : ButtonStyle)
+			.IsEnabled(InArgs._IsEnabled)
+			.ToolTipText(InArgs._ToolTipText)
+			.HAlign(HorizontalContentAlignment)
+			.VAlign(VAlign_Center)
+			.OnClicked(InArgs._OnClicked)
+			[
+				ButtonContentContainer
+			]
+		];
+	}
+}
+
+void SActionButton::SetMenuContentWidgetToFocus(TWeakPtr<SWidget> InWidget)
+{
+	check(ComboButton.IsValid());
+	ComboButton->SetMenuContentWidgetToFocus(InWidget);
+}
+
+void SActionButton::SetIsMenuOpen(bool bInIsOpen, bool bInIsFocused)
+{
+	check(ComboButton.IsValid());
+	ComboButton->SetIsOpen(bInIsOpen, bInIsFocused);
+}
+
+FSlateColor SActionButton::GetIconColorAndOpacity() const
+{
+	switch (ActionButtonType.Get(EActionButtonType::Default))
+	{
+	case EActionButtonType::Default:
+	case EActionButtonType::Simple:
+	default:
+		return FSlateColor::UseForeground();
+
+	case EActionButtonType::Positive:
+		return FStyleColors::AccentGreen;
+
+	case EActionButtonType::Warning:
+		return FStyleColors::Warning;
+
+	case EActionButtonType::Error:
+		return FStyleColors::Error;
+	}
+}
diff --git a/Engine/Source/Developer/ToolWidgets/Private/SNegativeActionButton.cpp b/Engine/Source/Developer/ToolWidgets/Private/SNegativeActionButton.cpp
--- a/Engine/Source/Developer/ToolWidgets/Private/SNegativeActionButton.cpp
+++ b/Engine/Source/Developer/ToolWidgets/Private/SNegativeActionButton.cpp
@@ -1,98 +1,69 @@
 // Copyright Epic Games, Inc. All Rights Reserved.
 
 #include "SNegativeActionButton.h"
+
+#include "SActionButton.h"
+#include "ToolWidgetsStyle.h"
 #include "Widgets/Images/SImage.h"
 #include "Widgets/Input/SButton.h"
-#include "Styling/StyleColors.h"
 
-void SNegativeActionButton::Construct(const FArguments& InArgs)
+namespace UE::ToolWidgets::Private
 {
-	TAttribute<FText> Text = InArgs._Text;
-	TAttribute<EActionButtonStyle> ActionButtonStyle = InArgs._ActionButtonStyle;
-	const FSlateBrush* Icon;
-
-	if (InArgs._Icon.IsSet())
+	EActionButtonType ActionButtonStyleToActionButtonType(const EActionButtonStyle InStyle)
 	{
-		Icon = InArgs._Icon.Get();
-	}
-	else
-	{
-		Icon = FAppStyle::Get().GetBrush(ActionButtonStyle.Get() == EActionButtonStyle::Error ? "Icons.Error" : "Icons.Warning");
-	}
+		static TMap<EActionButtonStyle, EActionButtonType> Lookup = {
+			{ EActionButtonStyle::Warning, EActionButtonType::Warning },
+			{ EActionButtonStyle::Error, EActionButtonType::Error }
+		};
 
-	TSharedRef<SHorizontalBox> ButtonContent = SNew(SHorizontalBox)
-		+ SHorizontalBox::Slot()
-		.HAlign(HAlign_Center)
-		.VAlign(VAlign_Center)
-		.AutoWidth()
-		[
-			SNew(SImage)
-			.Image(Icon)
-			.ColorAndOpacity(ActionButtonStyle.Get() == EActionButtonStyle::Error ? FStyleColors::Error : FStyleColors::Warning)
-		]
-		+ SHorizontalBox::Slot()
-		.Padding(FMargin(3, 0, 0, 0))
-		.VAlign(VAlign_Center)
-		.AutoWidth()
-		[
-			SNew(STextBlock)
-			.TextStyle(FAppStyle::Get(), "SmallButtonText")
-			.Text(InArgs._Text)
-			.Visibility_Lambda([Text]() { return Text.Get(FText::GetEmpty()).IsEmpty() ? EVisibility::Collapsed : EVisibility::Visible; })
-		];
+		if (const EActionButtonType* FoundEnum = Lookup.Find(InStyle))
+		{
+			return *FoundEnum;
+		}
 
-	if (InArgs._OnClicked.IsBound())
-	{
-		ChildSlot
-		[
-			SAssignNew(Button, SButton)
-			.ForegroundColor(FSlateColor::UseStyle())
-			.IsEnabled(InArgs._IsEnabled)
-			.ToolTipText(InArgs._ToolTipText)
-			.HAlign(HAlign_Center)
-			.VAlign(VAlign_Center)
-			.OnClicked(InArgs._OnClicked)
-			[
-				ButtonContent
-			]
-		];
+		return EActionButtonType::Default;
 	}
-	else
-	{
-		ChildSlot
+}
+
+void SNegativeActionButton::Construct(const FArguments& InArgs)
+{
+	// @note: Delegate-bound ActionButtonStyle is not supported, it should be provided at construction.
+	const EActionButtonType ActionButtonType = UE::ToolWidgets::Private::ActionButtonStyleToActionButtonType(InArgs._ActionButtonStyle.Get(EActionButtonStyle::Warning));
+
+	ActionButton =
+		SNew(SActionButton)
+		.ActionButtonStyle(
+			InArgs._ActionButtonStyle.Get(EActionButtonStyle::Warning) == EActionButtonStyle::Warning
+			? &UE::ToolWidgets::FToolWidgetsStyle::Get().GetWidgetStyle<FActionButtonStyle>("NegativeActionButton.Warning")
+			: &UE::ToolWidgets::FToolWidgetsStyle::Get().GetWidgetStyle<FActionButtonStyle>("NegativeActionButton.Error"))
+		.ActionButtonType(ActionButtonType)
+		.Text(InArgs._Text)
+		.Icon(InArgs._Icon)
+		.OnClicked(InArgs._OnClicked)
+		.OnGetMenuContent(InArgs._OnGetMenuContent)
+		.OnComboBoxOpened(InArgs._OnComboBoxOpened)
+		.OnMenuOpenChanged(InArgs._OnMenuOpenChanged)
+		.MenuContent()
 		[
-			SAssignNew(ComboButton, SComboButton)
-			.HasDownArrow(false)
-			.ContentPadding(FMargin(2.0f, 3.0f))
-			.ButtonStyle(&FAppStyle::Get().GetWidgetStyle<FButtonStyle>("Button"))
-			.ForegroundColor(FSlateColor::UseStyle())
-			.IsEnabled(InArgs._IsEnabled)
-			.ToolTipText(InArgs._ToolTipText)
-			.HAlign(HAlign_Center)
-			.VAlign(VAlign_Center)
-			.ButtonContent()
-			[
-				ButtonContent
-			]
-			.MenuContent()
-			[
-				InArgs._MenuContent.Widget
-			]
-			.OnGetMenuContent(InArgs._OnGetMenuContent)
-			.OnMenuOpenChanged(InArgs._OnMenuOpenChanged)
-			.OnComboBoxOpened(InArgs._OnComboBoxOpened)
+			InArgs._MenuContent.Widget
 		];
-	}
+
+	ChildSlot
+	[
+		ActionButton.ToSharedRef()
+	];
 }
 
-void SNegativeActionButton::SetMenuContentWidgetToFocus(TWeakPtr<SWidget> Widget)
+void SNegativeActionButton::SetMenuContentWidgetToFocus(TWeakPtr<SWidget> InWidget)
 {
-	check(ComboButton.IsValid());
-	ComboButton->SetMenuContentWidgetToFocus(Widget);
+	check(ActionButton.IsValid());
+
+	ActionButton->SetMenuContentWidgetToFocus(InWidget);
 }
 
-void SNegativeActionButton::SetIsMenuOpen(bool bIsOpen, bool bIsFocused)
+void SNegativeActionButton::SetIsMenuOpen(bool bInIsOpen, bool bInIsFocused)
 {
-	check(ComboButton.IsValid());
-	ComboButton->SetIsOpen(bIsOpen, bIsFocused);
+	check(ActionButton.IsValid());
+
+	ActionButton->SetIsMenuOpen(bInIsOpen, bInIsFocused);
 }
diff --git a/Engine/Source/Developer/ToolWidgets/Private/SPositiveActionButton.cpp b/Engine/Source/Developer/ToolWidgets/Private/SPositiveActionButton.cpp
--- a/Engine/Source/Developer/ToolWidgets/Private/SPositiveActionButton.cpp
+++ b/Engine/Source/Developer/ToolWidgets/Private/SPositiveActionButton.cpp
@@ -1,89 +1,46 @@
 // Copyright Epic Games, Inc. All Rights Reserved.
 
 #include "SPositiveActionButton.h"
+
+#include "SActionButton.h"
+#include "ToolWidgetsStyle.h"
 #include "Widgets/Images/SImage.h"
 #include "Widgets/Input/SButton.h"
-#include "Styling/StyleColors.h"
 
 void SPositiveActionButton::Construct(const FArguments& InArgs)
 {
 	check(InArgs._Icon.IsSet());
 
-	TAttribute<FText> Text = InArgs._Text;
-	
-	TSharedRef<SHorizontalBox> ButtonContent = SNew(SHorizontalBox)
-		+ SHorizontalBox::Slot()
-		.HAlign(HAlign_Center)
-		.VAlign(VAlign_Center)
-		.AutoWidth()
-		[
-			SNew(SImage)
-			.Image(InArgs._Icon)
-			.ColorAndOpacity(FStyleColors::AccentGreen)
-		]
-		+ SHorizontalBox::Slot()
-		.Padding(FMargin(3, 0, 0, 0))
-		.VAlign(VAlign_Center)
-		.AutoWidth()
+	ActionButton =
+		SNew(SActionButton)
+		.ActionButtonStyle(&UE::ToolWidgets::FToolWidgetsStyle::Get().GetWidgetStyle<FActionButtonStyle>("PositiveActionButton"))
+		.Text(InArgs._Text)
+		.Icon(InArgs._Icon)
+		.OnClicked(InArgs._OnClicked)
+		.OnGetMenuContent(InArgs._OnGetMenuContent)
+		.OnComboBoxOpened(InArgs._OnComboBoxOpened)
+		.OnMenuOpenChanged(InArgs._OnMenuOpenChanged)
+		.MenuContent()
 		[
-			SNew(STextBlock)
-			.TextStyle(FAppStyle::Get(), "SmallButtonText")
-			.Text(InArgs._Text)
-			.Visibility_Lambda([Text]() { return Text.Get(FText::GetEmpty()).IsEmpty() ? EVisibility::Collapsed : EVisibility::Visible; })
+			InArgs._MenuContent.Widget
 		];
 
-	if (InArgs._OnClicked.IsBound())
-	{
-		ChildSlot
-		[
-			SAssignNew(Button, SButton)
-			.ForegroundColor(FSlateColor::UseStyle())
-			.IsEnabled(InArgs._IsEnabled)
-			.ToolTipText(InArgs._ToolTipText)
-			.HAlign(HAlign_Center)
-			.VAlign(VAlign_Center)
-			.OnClicked(InArgs._OnClicked)
-			[
-				ButtonContent
-			]
-		];
-	}
-	else
-	{
-		ChildSlot
-		[
-			SAssignNew(ComboButton, SComboButton)
-			.HasDownArrow(false)
-			.ContentPadding(FMargin(2.0f, 3.0f))
-			.ButtonStyle(&FAppStyle::Get().GetWidgetStyle<FButtonStyle>("Button"))
-			.ForegroundColor(FSlateColor::UseStyle())
-			.IsEnabled(InArgs._IsEnabled)
-			.ToolTipText(InArgs._ToolTipText)
-			.HAlign(HAlign_Center)
-			.VAlign(VAlign_Center)
-			.ButtonContent()
-			[
-				ButtonContent
-			]
-			.MenuContent()
-			[
-				InArgs._MenuContent.Widget
-			]
-			.OnGetMenuContent(InArgs._OnGetMenuContent)
-			.OnMenuOpenChanged(InArgs._OnMenuOpenChanged)
-			.OnComboBoxOpened(InArgs._OnComboBoxOpened)
-		];
-	}
+	ChildSlot
+	[
+		ActionButton.ToSharedRef()
+	];
 }
 
-void SPositiveActionButton::SetMenuContentWidgetToFocus(TWeakPtr<SWidget> Widget)
+void SPositiveActionButton::SetMenuContentWidgetToFocus(TWeakPtr<SWidget> InWidget)
 {
-	check(ComboButton.IsValid());
-	ComboButton->SetMenuContentWidgetToFocus(Widget);
+	check(ActionButton.IsValid());
+
+	ActionButton->SetMenuContentWidgetToFocus(InWidget);
 }
 
-void SPositiveActionButton::SetIsMenuOpen(bool bIsOpen, bool bIsFocused)
+void SPositiveActionButton::SetIsMenuOpen(bool bInIsOpen, bool bInIsFocused)
 {
-	check(ComboButton.IsValid());
-	ComboButton->SetIsOpen(bIsOpen, bIsFocused);
+	check(ActionButton.IsValid());
+
+	ActionButton->SetIsMenuOpen(bInIsOpen, bInIsFocused);
 }
diff --git a/Engine/Source/Developer/ToolWidgets/Private/SPrimaryButton.cpp b/Engine/Source/Developer/ToolWidgets/Private/SPrimaryButton.cpp
--- a/Engine/Source/Developer/ToolWidgets/Private/SPrimaryButton.cpp
+++ b/Engine/Source/Developer/ToolWidgets/Private/SPrimaryButton.cpp
@@ -1,73 +1,22 @@
 // Copyright Epic Games, Inc. All Rights Reserved.
 
 #include "SPrimaryButton.h"
+
+#include "SActionButton.h"
+#include "ToolWidgetsSlateTypes.h"
+#include "ToolWidgetsStyle.h"
 #include "Widgets/SBoxPanel.h"
 #include "Widgets/Images/SImage.h"
 #include "Widgets/Text/STextBlock.h"
 
 void SPrimaryButton::Construct(const FArguments& InArgs)
 {
-	// To know if the icon is set, check explicitely for a brush or delegate, because
-	// if a caller passed nullptr, that would result in _Icon.IsSet() returning true.
-	// (Specifically checking IsBound() first so that if it is, we don't invoke it by calling Get())
-	const bool bIconSet = InArgs._Icon.IsBound() || InArgs._Icon.Get();
-
-	// If a delegate is bound, we'll store it and pass our own delegate to the Image.
-	// This will allow us to update the layout as an icon comes and goes.
-	TAttribute<const FSlateBrush*> IconAttribute;
-	if (InArgs._Icon.IsBound())
-	{
-		OwnerDelegate = InArgs._Icon.GetBinding();
-
-		TAttribute<const FSlateBrush*>::FGetter Delegate = TAttribute<const FSlateBrush*>::FGetter::CreateSP(this, &SPrimaryButton::OnGetBrush);		
-		IconAttribute.Bind(MoveTemp(Delegate));
-	}
-	else
-	{
-		IconAttribute = InArgs._Icon;
-	}
-
-	SButton::Construct(SButton::FArguments()
-		.ButtonStyle(&FAppStyle::Get().GetWidgetStyle<FButtonStyle>(bIconSet ? "PrimaryButtonLabelAndIcon" : "PrimaryButton"))
+	ChildSlot
+	[
+		SNew(SActionButton)
+		.ActionButtonStyle(&UE::ToolWidgets::FToolWidgetsStyle::Get().GetWidgetStyle<FActionButtonStyle>("PrimaryButton"))
+		.Text(InArgs._Text)
+		.Icon(InArgs._Icon)
 		.OnClicked(InArgs._OnClicked)
-		.ForegroundColor(FSlateColor::UseStyle())
-		.HAlign(HAlign_Center)
-		[
-			SAssignNew(HorizontalBox, SHorizontalBox)
-			+SHorizontalBox::Slot()
-			.AutoWidth()
-			.VAlign(VAlign_Center)
-			.Padding(FMargin(0,0,bIconSet ? 3 : 0,0))
-			[
-				SNew(SImage)
-				.ColorAndOpacity(FSlateColor::UseForeground())
-				.Image(IconAttribute)
-				.Visibility(bIconSet ? EVisibility::HitTestInvisible : EVisibility::Collapsed)
-			]
-			+ SHorizontalBox::Slot()
-			.VAlign(VAlign_Center)
-			.AutoWidth()
-			[
-				SNew(STextBlock)
-				.TextStyle(&FAppStyle::Get().GetWidgetStyle<FTextBlockStyle>("PrimaryButtonText"))
-				.Text(InArgs._Text)
-				.Visibility(InArgs._Text.IsSet() ? EVisibility::HitTestInvisible : EVisibility::Collapsed)
-			]
-		]
-	);
+	];
 }
-
-const FSlateBrush* SPrimaryButton::OnGetBrush()
-{
-	const FSlateBrush* Brush = nullptr;
-	if (OwnerDelegate.IsBound())
-	{
-		Brush = OwnerDelegate.Execute();
-	}
-
-	// Update the button's styling based on whether a brush will be rendered or not
-	SetButtonStyle(&FAppStyle::Get().GetWidgetStyle<FButtonStyle>(Brush ? "PrimaryButtonLabelAndIcon" : "PrimaryButton"));
-	HorizontalBox->GetSlot(0).SetPadding(FMargin(0,0,Brush ? 3 : 0,0));
-
-	return Brush;
-}
\ No newline at end of file
diff --git a/Engine/Source/Developer/ToolWidgets/Private/SSimpleButton.cpp b/Engine/Source/Developer/ToolWidgets/Private/SSimpleButton.cpp
--- a/Engine/Source/Developer/ToolWidgets/Private/SSimpleButton.cpp
+++ b/Engine/Source/Developer/ToolWidgets/Private/SSimpleButton.cpp
@@ -1,38 +1,22 @@
 // Copyright Epic Games, Inc. All Rights Reserved.
 
 #include "SSimpleButton.h"
+
+#include "SActionButton.h"
+#include "ToolWidgetsStyle.h"
 #include "Widgets/SBoxPanel.h"
 #include "Widgets/Images/SImage.h"
 #include "Widgets/Text/STextBlock.h"
 
 void SSimpleButton::Construct(const FArguments& InArgs)
 {
-	SButton::Construct(SButton::FArguments()
+	ChildSlot
+	[
+		SNew(SActionButton)
+		.ActionButtonStyle(&UE::ToolWidgets::FToolWidgetsStyle::Get().GetWidgetStyle<FActionButtonStyle>("SimpleButton"))
 		.ButtonStyle(&FAppStyle::Get().GetWidgetStyle<FButtonStyle>(InArgs._Text.IsSet() ? "SimpleButtonLabelAndIcon" : "SimpleButton"))
+		.Text(InArgs._Text)
+		.Icon(InArgs._Icon)
 		.OnClicked(InArgs._OnClicked)
-		.ForegroundColor(FSlateColor::UseStyle())
-		[
-			SNew(SHorizontalBox)
-			+SHorizontalBox::Slot()
-			.AutoWidth()
-			.VAlign(VAlign_Center)
-			[
-				SNew(SImage)
-				.ColorAndOpacity(FSlateColor::UseForeground())
-				.Image(InArgs._Icon)
-				.Visibility(InArgs._Icon.IsSet() ? EVisibility::HitTestInvisible : EVisibility::Collapsed)
-			]
-			+ SHorizontalBox::Slot()
-			.VAlign(VAlign_Center)
-			.Padding(FMargin(3, 0, 0, 0))
-			.AutoWidth()
-			[
-				SNew(STextBlock)
-				.TextStyle(&FAppStyle::Get().GetWidgetStyle<FTextBlockStyle>("SmallButtonText"))
-				.Text(InArgs._Text)
-				.Visibility(InArgs._Text.IsSet() ? EVisibility::HitTestInvisible : EVisibility::Collapsed)
-			]
-		]
-	);
+	];
 }
-
diff --git a/Engine/Source/Developer/ToolWidgets/Private/SSimpleComboButton.cpp b/Engine/Source/Developer/ToolWidgets/Private/SSimpleComboButton.cpp
--- a/Engine/Source/Developer/ToolWidgets/Private/SSimpleComboButton.cpp
+++ b/Engine/Source/Developer/ToolWidgets/Private/SSimpleComboButton.cpp
@@ -1,49 +1,48 @@
 // Copyright Epic Games, Inc. All Rights Reserved.
 
 #include "SSimpleComboButton.h"
+
+#include "SActionButton.h"
 #include "Widgets/SBoxPanel.h"
 #include "Widgets/Images/SImage.h"
 #include "Widgets/Text/STextBlock.h"
 #include "Styling/SlateTypes.h"
+#include "ToolWidgetsStyle.h"
 
 void SSimpleComboButton::Construct(const FArguments& InArgs)
 {
-
-	TAttribute<FText> Text = InArgs._Text;
-	FTextBlockStyle TextStyle = InArgs._UsesSmallText ? FAppStyle::GetWidgetStyle<FTextBlockStyle>("SmallText") : FAppStyle::GetWidgetStyle<FTextBlockStyle>("SmallButtonText");
-
-	TSharedRef<SHorizontalBox> ButtonContent = SNew(SHorizontalBox)
-		+ SHorizontalBox::Slot()
-		.HAlign(HAlign_Center)
-		.VAlign(VAlign_Center)
-		.AutoWidth()
+	ActionButton = SNew(SActionButton)
+		.ActionButtonStyle(&UE::ToolWidgets::FToolWidgetsStyle::Get().GetWidgetStyle<FActionButtonStyle>("SimpleComboButton"))
+		.Text(InArgs._Text)
+		.TextBlockStyle(
+			InArgs._UsesSmallText
+			? &FAppStyle::GetWidgetStyle<FTextBlockStyle>("SmallText")
+			: &FAppStyle::GetWidgetStyle<FTextBlockStyle>("SmallButtonText"))
+		.Icon(InArgs._Icon)
+		.MenuContent()
 		[
-			SNew(SImage)
-			.Image(InArgs._Icon)
-			.ColorAndOpacity(FSlateColor::UseForeground())
+			InArgs._MenuContent.Widget
 		]
-		+ SHorizontalBox::Slot()
-		.Padding(FMargin(3, 0, 0, 0))
-		.VAlign(VAlign_Center)
-		.FillWidth(1.0f)
-		[
-			SNew(STextBlock)
-			.TextStyle(&TextStyle)
-			.Text(InArgs._Text)
-			.Visibility_Lambda([Text]() { return Text.Get(FText::GetEmpty()).IsEmpty() ? EVisibility::Collapsed : EVisibility::Visible; })
-			.Clipping(EWidgetClipping::OnDemand)
-		];
-
-	SComboButton::Construct(SComboButton::FArguments()
-		.HasDownArrow(InArgs._HasDownArrow)
-		.ContentPadding(FMargin(2.0f, 2.0f))
-		.ComboButtonStyle(FAppStyle::Get(), "SimpleComboButton")
-		.ForegroundColor(FSlateColor::UseStyle())
-		.ButtonContent() [ButtonContent]
-		.MenuContent()[InArgs._MenuContent.Widget]
 		.OnGetMenuContent(InArgs._OnGetMenuContent)
 		.OnMenuOpenChanged(InArgs._OnMenuOpenChanged)
-		.OnComboBoxOpened(InArgs._OnComboBoxOpened)
-	);
+		.OnComboBoxOpened(InArgs._OnComboBoxOpened);
+
+	ChildSlot
+	[
+		ActionButton.ToSharedRef()
+	];
+}
+
+void SSimpleComboButton::SetMenuContentWidgetToFocus(TWeakPtr<SWidget> InWidget)
+{
+	check(ActionButton.IsValid());
+
+	ActionButton->SetMenuContentWidgetToFocus(InWidget);
 }
 
+void SSimpleComboButton::SetIsMenuOpen(bool bInIsOpen, bool bInIsFocused)
+{
+	check(ActionButton.IsValid());
+
+	ActionButton->SetIsMenuOpen(bInIsOpen, bInIsFocused);
+}
diff --git a/Engine/Source/Developer/ToolWidgets/Private/ToolWidgetsSlateTypes.cpp b/Engine/Source/Developer/ToolWidgets/Private/ToolWidgetsSlateTypes.cpp
new file mode 100644
--- /dev/null
+++ b/Engine/Source/Developer/ToolWidgets/Private/ToolWidgetsSlateTypes.cpp
@@ -0,0 +1,175 @@
+// Copyright Epic Games, Inc. All Rights Reserved.
+
+#include "ToolWidgetsSlateTypes.h"
+
+namespace UE::ToolWidgets::Private
+{
+	inline FName LexToName(const EActionButtonType ActionButtonType)
+	{
+		static const TMap<EActionButtonType, const FName> Lookup = {
+			{ EActionButtonType::Default, TEXT("Default") },
+			{ EActionButtonType::Simple, TEXT("Simple") },
+			{ EActionButtonType::Primary, TEXT("Primary") },
+			{ EActionButtonType::Positive, TEXT("Positive") },
+			{ EActionButtonType::Warning, TEXT("Warning") },
+			{ EActionButtonType::Error, TEXT("Error") },
+		};
+
+		if (const FName* FoundName = Lookup.Find(ActionButtonType))
+		{
+			return *FoundName;
+		}
+
+		return TEXT("Invalid");
+	}
+
+	inline EActionButtonType NameToEnum(const FName ActionButtonTypeName)
+	{
+		static const TMap<const FName, EActionButtonType> Lookup = {
+			{ TEXT("Default"), EActionButtonType::Default },
+			{ TEXT("Simple"), EActionButtonType::Simple },
+			{ TEXT("Primary"), EActionButtonType::Primary },
+			{ TEXT("Positive"), EActionButtonType::Positive },
+			{ TEXT("Warning"), EActionButtonType::Warning },
+			{ TEXT("Error"), EActionButtonType::Error },
+		};
+
+		if (const EActionButtonType* FoundEnum = Lookup.Find(ActionButtonTypeName))
+		{
+			return *FoundEnum;
+		}
+
+		return EActionButtonType::Default;
+	}
+
+}
+
+const FName FActionButtonStyle::TypeName(TEXT("FActionButtonStyle"));
+
+FActionButtonStyle::FActionButtonStyle()
+	: bHasDownArrow(false)
+	, HorizontalContentAlignment(HAlign_Center)
+{
+}
+
+void FActionButtonStyle::GetResources(TArray<const FSlateBrush*>& OutBrushes) const
+{
+	if (IconBrush.IsSet())
+	{
+		OutBrushes.Add(&IconBrush.GetValue());
+	}
+
+	ButtonStyle.GetResources(OutBrushes);
+	ComboButtonStyle.GetResources(OutBrushes);
+	TextBlockStyle.GetResources(OutBrushes);
+}
+
+const FActionButtonStyle& FActionButtonStyle::GetDefault()
+{
+	static FActionButtonStyle Default;
+	return Default;
+}
+
+const FName FActionButtonStyle::GetTypeName() const
+{
+	return TypeName;
+}
+
+EActionButtonType FActionButtonStyle::GetActionButtonType() const
+{
+	return UE::ToolWidgets::Private::NameToEnum(ActionButtonType);
+}
+
+FActionButtonStyle& FActionButtonStyle::SetActionButtonType(const EActionButtonType InActionButtonType)
+{
+	ActionButtonType = UE::ToolWidgets::Private::LexToName(InActionButtonType);
+	return *this;
+}
+
+FActionButtonStyle& FActionButtonStyle::SetButtonStyle(const FButtonStyle& InButtonStyle)
+{
+	ButtonStyle = InButtonStyle;
+	return *this;
+}
+
+const FButtonStyle& FActionButtonStyle::GetIconButtonStyle() const
+{
+	return IconButtonStyle.Get(ButtonStyle);
+}
+
+FActionButtonStyle& FActionButtonStyle::SetIconButtonStyle(const FButtonStyle& InButtonStyle)
+{
+	IconButtonStyle = InButtonStyle;
+	return *this;
+}
+
+FActionButtonStyle& FActionButtonStyle::SetButtonContentPadding(const FMargin& InContentPadding)
+{
+	ButtonContentPadding = InContentPadding;
+	return *this;
+}
+
+FActionButtonStyle& FActionButtonStyle::SetComboButtonStyle(const FComboButtonStyle& InComboButtonStyle)
+{
+	ComboButtonStyle = InComboButtonStyle;
+	return *this;
+}
+
+FActionButtonStyle& FActionButtonStyle::SetHasDownArrow(const bool& bInHasDownArrow)
+{
+	bHasDownArrow = bInHasDownArrow;
+	return *this;
+}
+
+FActionButtonStyle& FActionButtonStyle::SetComboButtonContentPadding(const FMargin& InContentPadding)
+{
+	ComboButtonContentPadding = InContentPadding;
+	return *this;
+}
+
+FActionButtonStyle& FActionButtonStyle::SetHorizontalContentAlignment(const EHorizontalAlignment& InAlignment)
+{
+	HorizontalContentAlignment = InAlignment;
+	return *this;
+}
+
+FActionButtonStyle& FActionButtonStyle::SetTextBlockStyle(const FTextBlockStyle& InTextBlockStyle)
+{
+	TextBlockStyle = InTextBlockStyle;
+	return *this;
+}
+
+FActionButtonStyle& FActionButtonStyle::SetIconBrush(const FSlateBrush& InIconBrush)
+{
+	IconBrush = InIconBrush;
+	return *this;
+}
+
+FActionButtonStyle& FActionButtonStyle::SetIconColorAndOpacity(const FSlateColor& InIconColorAndOpacity)
+{
+	IconColorAndOpacity = InIconColorAndOpacity;
+	return *this;
+}
+
+void FActionButtonStyle::UnlinkColors()
+{
+	ButtonStyle.UnlinkColors();
+	ComboButtonStyle.UnlinkColors();
+	TextBlockStyle.UnlinkColors();
+}
+
+#if WITH_EDITOR
+const TArray<FName>& UToolSlateWidgetTypesFunctionLibrary::GetActionButtonTypeNames()
+{
+	static const TArray<FName> ActionButtonTypeNames = {
+		TEXT("Default"),
+		TEXT("Simple"),
+		TEXT("Primary"),
+		TEXT("Positive"),
+		TEXT("Warning"),
+		TEXT("Error"),
+	};
+
+	return ActionButtonTypeNames;
+}
+#endif
diff --git a/Engine/Source/Developer/ToolWidgets/Private/ToolWidgetsStyle.cpp b/Engine/Source/Developer/ToolWidgets/Private/ToolWidgetsStyle.cpp
new file mode 100644
--- /dev/null
+++ b/Engine/Source/Developer/ToolWidgets/Private/ToolWidgetsStyle.cpp
@@ -0,0 +1,152 @@
+// Copyright Epic Games, Inc. All Rights Reserved.
+
+#include "ToolWidgetsStyle.h"
+
+#include "Styling/AppStyle.h"
+#include "Styling/SlateStyleRegistry.h"
+#include "Styling/StyleColors.h"
+#include "ToolWidgetsSlateTypes.h"
+
+namespace UE::ToolWidgets
+{
+	namespace Private
+	{
+		/** Compensates for Button Style padding. */
+		static FMargin ModifyContentPadding(const FButtonStyle& InButtonStyle, const FMargin& InContentPadding)
+		{
+			FMargin ContentPadding = InContentPadding;
+			ContentPadding.Top = ContentPadding.Bottom = InButtonStyle.NormalPadding.Top;
+			return ContentPadding;
+		}
+
+		static FMargin ModifyContentPadding(const FComboButtonStyle& InComboButtonStyle, const FMargin& InContentPadding)
+		{
+			return ModifyContentPadding(InComboButtonStyle.ButtonStyle,	InContentPadding);
+		}
+	}
+	
+	FName FToolWidgetsStyle::StyleName("ToolWidgets");
+
+	FToolWidgetsStyle& FToolWidgetsStyle::Get()
+	{
+		static FToolWidgetsStyle Instance;
+		return Instance;
+	}
+
+	FToolWidgetsStyle::FToolWidgetsStyle()
+		: FSlateStyleSet(StyleName)
+	{
+		SetParentStyleName(FAppStyle::GetAppStyleSetName());
+
+		SetContentRoot(FPaths::EngineContentDir() / TEXT("Editor/Slate"));
+		SetCoreContentRoot(FPaths::EngineContentDir() / TEXT("Slate"));
+
+		// SActionButton
+		{
+			using namespace Private;
+		
+			static const FMargin DefaultButtonContentPadding = FMargin(FActionButton::DefaultHorizontalPadding, FActionButton::DefaultVerticalPadding);
+
+			const FButtonStyle& DefaultButtonStyle = GetWidgetStyle<FButtonStyle>("Button");
+			const FComboButtonStyle& DefaultComboButtonStyle = GetWidgetStyle<FComboButtonStyle>("ComboButton");
+
+			FActionButtonStyle ActionButton = FActionButtonStyle();
+			{
+				ActionButton
+					.SetActionButtonType(EActionButtonType::Default)
+					.SetButtonStyle(DefaultButtonStyle)
+					.SetButtonContentPadding(ModifyContentPadding(DefaultButtonStyle, DefaultButtonContentPadding))
+					.SetComboButtonStyle(DefaultComboButtonStyle)
+					.SetComboButtonContentPadding(ModifyContentPadding(DefaultButtonStyle, DefaultButtonContentPadding))
+					.SetHorizontalContentAlignment(HAlign_Center)
+					.SetTextBlockStyle(GetWidgetStyle<FTextBlockStyle>("SmallButtonText"))
+					.SetHasDownArrow(false);
+
+				Set("ActionButton", ActionButton);
+			}
+
+			FActionButtonStyle PositiveActionButton = ActionButton;
+			{
+				PositiveActionButton
+					.SetActionButtonType(EActionButtonType::Positive)
+					.SetIconBrush(*FAppStyle::Get().GetBrush("Icons.Plus"))
+					.SetIconColorAndOpacity(FStyleColors::AccentGreen);
+
+				Set("PositiveActionButton", PositiveActionButton);
+			}
+
+			FActionButtonStyle NegativeActionButton = ActionButton;
+			{
+				FActionButtonStyle NegativeActionButtonWarning = NegativeActionButton;
+
+				NegativeActionButtonWarning
+					.SetActionButtonType(EActionButtonType::Warning)
+					.SetIconBrush(*FAppStyle::Get().GetBrush("Icons.Warning"));
+
+				Set("NegativeActionButton.Warning", NegativeActionButtonWarning);
+
+				FActionButtonStyle NegativeActionButtonError = NegativeActionButton;
+
+				NegativeActionButtonError
+					.SetActionButtonType(EActionButtonType::Error)
+					.SetIconBrush(*FAppStyle::Get().GetBrush("Icons.Error"));
+
+				Set("NegativeActionButton.Error", NegativeActionButtonError);
+			}
+
+			FActionButtonStyle SimpleComboButton = ActionButton;
+			{
+				const FComboButtonStyle& ComboButtonStyle = GetWidgetStyle<FComboButtonStyle>("SimpleComboButton");
+
+				SimpleComboButton
+					.SetActionButtonType(EActionButtonType::Simple)
+					.SetComboButtonStyle(ComboButtonStyle)
+					.SetButtonStyle(ComboButtonStyle.ButtonStyle)
+					.SetComboButtonContentPadding(ModifyContentPadding(ComboButtonStyle, DefaultButtonContentPadding))
+					.SetHorizontalContentAlignment(HAlign_Left)
+					.SetTextBlockStyle(GetWidgetStyle<FTextBlockStyle>("SmallText"))
+					.SetIconColorAndOpacity(FSlateColor::UseForeground());
+
+				Set("SimpleComboButton", SimpleComboButton);
+			}
+
+			FActionButtonStyle SimpleButton = ActionButton;
+			{
+				const FButtonStyle& ButtonStyle = GetWidgetStyle<FButtonStyle>("SimpleButtonLabelAndIcon");
+
+				SimpleButton
+					.SetActionButtonType(EActionButtonType::Simple)
+					.SetButtonStyle(ButtonStyle)
+					.SetButtonContentPadding(ModifyContentPadding(ButtonStyle, DefaultButtonContentPadding))
+					.SetHorizontalContentAlignment(HAlign_Left)
+					.SetTextBlockStyle(GetWidgetStyle<FTextBlockStyle>("SmallButtonText"))
+					.SetIconColorAndOpacity(FSlateColor::UseForeground());
+
+				Set("SimpleButton", SimpleButton);
+			}
+
+			FActionButtonStyle PrimaryButton = ActionButton;
+			{
+				const FButtonStyle& ButtonStyle = GetWidgetStyle<FButtonStyle>("PrimaryButton");
+
+				PrimaryButton
+					.SetActionButtonType(EActionButtonType::Primary)
+					.SetButtonStyle(ButtonStyle)
+					.SetButtonContentPadding(ModifyContentPadding(ButtonStyle, DefaultButtonContentPadding))
+					.SetComboButtonContentPadding(ModifyContentPadding(ButtonStyle, DefaultButtonContentPadding))
+					.SetIconButtonStyle(GetWidgetStyle<FButtonStyle>("PrimaryButtonLabelAndIcon"))
+					.SetTextBlockStyle(GetWidgetStyle<FTextBlockStyle>("PrimaryButtonText"))
+					.SetIconColorAndOpacity(FSlateColor::UseForeground());
+
+				Set("PrimaryButton", PrimaryButton);
+			}
+		}
+
+		FSlateStyleRegistry::RegisterSlateStyle(*this);
+	}
+
+	FToolWidgetsStyle::~FToolWidgetsStyle()
+	{
+		FSlateStyleRegistry::UnRegisterSlateStyle(*this);
+	}
+}
diff --git a/Engine/Source/Developer/ToolWidgets/Public/SActionButton.h b/Engine/Source/Developer/ToolWidgets/Public/SActionButton.h
new file mode 100644
--- /dev/null
+++ b/Engine/Source/Developer/ToolWidgets/Public/SActionButton.h
@@ -0,0 +1,106 @@
+// Copyright Epic Games, Inc. All Rights Reserved.
+
+#pragma once
+
+#include "Framework/SlateDelegates.h"
+#include "ToolWidgetsSlateTypes.h"
+#include "Widgets/DeclarativeSyntaxSupport.h"
+#include "Widgets/SCompoundWidget.h"
+#include "Widgets/Input/SComboButton.h"
+
+class TOOLWIDGETS_API IActionButton
+{
+public:
+	virtual ~IActionButton() = default;
+
+	virtual void SetMenuContentWidgetToFocus(TWeakPtr<SWidget> InWidget) = 0;
+	virtual void SetIsMenuOpen(bool bInIsOpen, bool bInIsFocused) = 0;
+};
+
+/** A Button that is used to call out/highlight an option. It can also be used to open a menu. */
+class TOOLWIDGETS_API SActionButton
+	: public SCompoundWidget
+	, public IActionButton
+{
+public:
+	SLATE_BEGIN_ARGS(SActionButton)
+		: _ActionButtonStyle(nullptr)
+		, _ButtonStyle(nullptr)
+		, _IconButtonStyle(nullptr)
+		, _ComboButtonStyle(nullptr)
+		, _TextBlockStyle(nullptr)
+		{}
+
+		SLATE_STYLE_ARGUMENT(FActionButtonStyle, ActionButtonStyle)
+
+		/** Used to describe the intent of this button. */
+		SLATE_ATTRIBUTE(EActionButtonType, ActionButtonType)
+
+		/** Optionally specify the button style, which may override visual properties determined by ActionButtonStyle. */
+		SLATE_STYLE_ARGUMENT(FButtonStyle, ButtonStyle)
+
+		/** Optionally specify the button style when an icon is present, which may override visual properties determined by ActionButtonStyle. */
+		SLATE_STYLE_ARGUMENT(FButtonStyle, IconButtonStyle)
+
+		/** Spacing between button's border and the content. */
+		SLATE_ATTRIBUTE(FMargin, ButtonContentPadding)
+
+		/** Optionally specify the combo button style, which may override visual properties determined by ActionButtonStyle. */
+		SLATE_STYLE_ARGUMENT(FComboButtonStyle, ComboButtonStyle)
+
+		/** Whether to show a down arrow for the combo button. Default is determined by the style. */
+		SLATE_ARGUMENT(TOptional<bool>, HasDownArrow)
+
+		/** Horizontal Content alignment within the button. Default is determined by the style. */
+		SLATE_ARGUMENT(TOptional<EHorizontalAlignment>, HorizontalContentAlignment)
+
+		/** The text to display in the button. */
+		SLATE_ATTRIBUTE(FText, Text)
+
+		/** The style of the text block, which dictates the font, color, and shadow options. */
+		SLATE_STYLE_ARGUMENT(FTextBlockStyle, TextBlockStyle)
+
+		/** Optionally specify the Icon to display in the button. */
+		SLATE_ATTRIBUTE(const FSlateBrush*, Icon)
+
+		/** Icon Color/Tint, defaults to White. */
+		SLATE_ATTRIBUTE(FSlateColor, IconColorAndOpacity)
+
+		#pragma region Button
+		/** The clicked handler. Note that if this is set, the button will behave as though it were just a button.
+		 * This means that OnGetMenuContent, OnComboBoxOpened and OnMenuOpenChanged will all be ignored, since there is no menu.
+		 */
+		SLATE_EVENT(FOnClicked, OnClicked)
+		#pragma endregion Button
+
+		#pragma region ComboButton
+		/** The static menu content widget. */
+		SLATE_NAMED_SLOT(FArguments, MenuContent)
+
+		SLATE_EVENT(FOnGetContent, OnGetMenuContent)
+		SLATE_EVENT(FOnComboBoxOpened, OnComboBoxOpened)
+		SLATE_EVENT(FOnIsOpenChanged, OnMenuOpenChanged)
+		#pragma endregion ComboButton
+	SLATE_END_ARGS()
+
+	void Construct(const FArguments& InArgs);
+
+	//~ Begin IActionButton
+	virtual void SetMenuContentWidgetToFocus(TWeakPtr<SWidget> InWidget) override;
+	virtual void SetIsMenuOpen(bool bInIsOpen, bool bInIsFocused) override;
+	//~ End IActionButton
+
+protected:
+	virtual FSlateColor GetIconColorAndOpacity() const;
+
+protected:
+	TSharedPtr<SComboButton> ComboButton;
+	TSharedPtr<class SButton> Button;
+
+	TAttribute<EActionButtonType> ActionButtonType = EActionButtonType::Default;
+	const FActionButtonStyle* ActionButtonStyle = nullptr;
+	const FButtonStyle* ButtonStyle = nullptr;
+	const FButtonStyle* IconButtonStyle = nullptr;
+	const FComboButtonStyle* ComboButtonStyle = nullptr;
+	const FTextBlockStyle* TextBlockStyle = nullptr;
+};
diff --git a/Engine/Source/Developer/ToolWidgets/Public/SNegativeActionButton.h b/Engine/Source/Developer/ToolWidgets/Public/SNegativeActionButton.h
--- a/Engine/Source/Developer/ToolWidgets/Public/SNegativeActionButton.h
+++ b/Engine/Source/Developer/ToolWidgets/Public/SNegativeActionButton.h
@@ -3,9 +3,10 @@
 #pragma once
 
 #include "Framework/SlateDelegates.h"
+#include "SActionButton.h"
 #include "Widgets/DeclarativeSyntaxSupport.h"
-#include "Widgets/SCompoundWidget.h"
 #include "Widgets/Input/SComboButton.h"
+#include "Widgets/SCompoundWidget.h"
 
 enum class EActionButtonStyle
 {
@@ -16,20 +17,22 @@ enum class EActionButtonStyle
 /** A Button that is used to call out/highlight a negative option (Warnings or Errors like Force Delete).
 *   It can also be used to open a menu.
 */
-class TOOLWIDGETS_API SNegativeActionButton : public SCompoundWidget
+class TOOLWIDGETS_API SNegativeActionButton
+	: public SCompoundWidget
+	, public IActionButton
 {
 public:
-
 	SLATE_BEGIN_ARGS(SNegativeActionButton)
 		: _ActionButtonStyle(EActionButtonStyle::Error)
-		, _Icon()
 		{}
 
+		/** Determines whether this is a Warning or Error. */
 		SLATE_ATTRIBUTE(EActionButtonStyle, ActionButtonStyle)
 
 		/** The text to display in the button. */
 		SLATE_ATTRIBUTE(FText, Text)
 
+		/** Optionally specify the Icon to display in the button. */
 		SLATE_ATTRIBUTE(const FSlateBrush*, Icon)
 
 		/** The clicked handler. Note that if this is set, the button will behave as though it were just a button.
@@ -43,18 +46,17 @@ public:
 		SLATE_EVENT(FOnGetContent, OnGetMenuContent)
 		SLATE_EVENT(FOnComboBoxOpened, OnComboBoxOpened)
 		SLATE_EVENT(FOnIsOpenChanged, OnMenuOpenChanged)
-
 	SLATE_END_ARGS()
 
-	SNegativeActionButton() {}
+	SNegativeActionButton() = default;
 
 	void Construct(const FArguments& InArgs);
 
-	void SetMenuContentWidgetToFocus(TWeakPtr<SWidget> Widget);
-	void SetIsMenuOpen(bool bIsOpen, bool bIsFocused);
+	//~ Begin IActionButton
+	virtual void SetMenuContentWidgetToFocus(TWeakPtr<SWidget> InWidget) override;
+	virtual void SetIsMenuOpen(bool bInIsOpen, bool bInIsFocused) override;
+	//~ End IActionButton
 
 private:
-
-	TSharedPtr<SComboButton> ComboButton;
-	TSharedPtr<class SButton> Button;
-};
\ No newline at end of file
+	TSharedPtr<SActionButton> ActionButton;
+};
diff --git a/Engine/Source/Developer/ToolWidgets/Public/SPositiveActionButton.h b/Engine/Source/Developer/ToolWidgets/Public/SPositiveActionButton.h
--- a/Engine/Source/Developer/ToolWidgets/Public/SPositiveActionButton.h
+++ b/Engine/Source/Developer/ToolWidgets/Public/SPositiveActionButton.h
@@ -3,18 +3,20 @@
 #pragma once
 
 #include "Framework/SlateDelegates.h"
+#include "SActionButton.h"
 #include "Widgets/DeclarativeSyntaxSupport.h"
 #include "Widgets/SCompoundWidget.h"
 #include "Widgets/Input/SComboButton.h"
 
 /** A Button that is used to call out/highlight a positive option (Add, Save etc). It can also be used to open a menu.
 */
-class TOOLWIDGETS_API SPositiveActionButton : public SCompoundWidget
+class TOOLWIDGETS_API SPositiveActionButton
+	: public SCompoundWidget
+	, public IActionButton
 {
 public:
-
-	SLATE_BEGIN_ARGS(SPositiveActionButton) :
-		_Icon(FAppStyle::Get().GetBrush("Icons.Plus"))
+	SLATE_BEGIN_ARGS(SPositiveActionButton)
+		: _Icon(FAppStyle::Get().GetBrush("Icons.Plus"))
 		{}
 
 		/** The text to display in the button. */
@@ -33,18 +35,17 @@ public:
 		SLATE_EVENT(FOnGetContent, OnGetMenuContent)
 		SLATE_EVENT(FOnComboBoxOpened, OnComboBoxOpened)
 		SLATE_EVENT(FOnIsOpenChanged, OnMenuOpenChanged)
-
 	SLATE_END_ARGS()
 
-	SPositiveActionButton() {}
+	SPositiveActionButton() = default;
 
 	void Construct(const FArguments& InArgs);
 
-	void SetMenuContentWidgetToFocus(TWeakPtr<SWidget> Widget);
-	void SetIsMenuOpen(bool bIsOpen, bool bIsFocused);
+	//~ Begin IActionButton
+	virtual void SetMenuContentWidgetToFocus(TWeakPtr<SWidget> InWidget) override;
+	virtual void SetIsMenuOpen(bool bInIsOpen, bool bInIsFocused) override;
+	//~ End IActionButton
 
 private:
-
-	TSharedPtr<SComboButton> ComboButton;
-	TSharedPtr<class SButton> Button;
-};
\ No newline at end of file
+	TSharedPtr<SActionButton> ActionButton;
+};
diff --git a/Engine/Source/Developer/ToolWidgets/Public/SPrimaryButton.h b/Engine/Source/Developer/ToolWidgets/Public/SPrimaryButton.h
--- a/Engine/Source/Developer/ToolWidgets/Public/SPrimaryButton.h
+++ b/Engine/Source/Developer/ToolWidgets/Public/SPrimaryButton.h
@@ -6,12 +6,13 @@
 #include "Widgets/DeclarativeSyntaxSupport.h"
 #include "Widgets/Input/SButton.h"
 
+class SActionButton;
 class SHorizontalBox;
 
-class TOOLWIDGETS_API SPrimaryButton : public SButton
+class TOOLWIDGETS_API SPrimaryButton
+	: public SCompoundWidget
 {
 public:
-
 	SLATE_BEGIN_ARGS(SPrimaryButton)
 	{}
 		/** The text to display in the button. */
@@ -24,15 +25,10 @@ public:
 		 * Called when the button is clicked  
 		 */
 		SLATE_EVENT(FOnClicked, OnClicked)
-
 	SLATE_END_ARGS()
 
 	void Construct(const FArguments& InArgs);
 
 private:
-	const FSlateBrush* OnGetBrush();
-
-private:
-	TAttribute<const FSlateBrush*>::FGetter OwnerDelegate;
-	TSharedPtr<SHorizontalBox> HorizontalBox;
-};
\ No newline at end of file
+	TSharedPtr<SActionButton> ActionButton;
+};
diff --git a/Engine/Source/Developer/ToolWidgets/Public/SSimpleButton.h b/Engine/Source/Developer/ToolWidgets/Public/SSimpleButton.h
--- a/Engine/Source/Developer/ToolWidgets/Public/SSimpleButton.h
+++ b/Engine/Source/Developer/ToolWidgets/Public/SSimpleButton.h
@@ -5,11 +5,12 @@
 #include "Framework/SlateDelegates.h"
 #include "Widgets/DeclarativeSyntaxSupport.h"
 #include "Widgets/Input/SButton.h"
+#include "Widgets/SCompoundWidget.h"
 
-class TOOLWIDGETS_API SSimpleButton : public SButton
+class TOOLWIDGETS_API SSimpleButton
+	: public SCompoundWidget
 {
 public:
-
 	SLATE_BEGIN_ARGS(SSimpleButton)
 	{}
 		/** The text to display in the button. */
@@ -19,10 +20,9 @@ public:
 
 		/** The clicked handler. */
 		SLATE_EVENT(FOnClicked, OnClicked)
-
 	SLATE_END_ARGS()
 
-	SSimpleButton() {}
+	SSimpleButton() = default;
 
 	void Construct(const FArguments& InArgs);
-};
\ No newline at end of file
+};
diff --git a/Engine/Source/Developer/ToolWidgets/Public/SSimpleComboButton.h b/Engine/Source/Developer/ToolWidgets/Public/SSimpleComboButton.h
--- a/Engine/Source/Developer/ToolWidgets/Public/SSimpleComboButton.h
+++ b/Engine/Source/Developer/ToolWidgets/Public/SSimpleComboButton.h
@@ -3,13 +3,16 @@
 #pragma once
 
 #include "Framework/SlateDelegates.h"
+#include "SActionButton.h"
 #include "Widgets/DeclarativeSyntaxSupport.h"
+#include "Widgets/SCompoundWidget.h"
 #include "Widgets/Input/SComboButton.h"
 
-class TOOLWIDGETS_API SSimpleComboButton : public SComboButton
+class TOOLWIDGETS_API SSimpleComboButton
+	: public SCompoundWidget
+	, public IActionButton
 {
 public:
-
 	SLATE_BEGIN_ARGS(SSimpleComboButton)
 		: _HasDownArrow(false)
 		, _UsesSmallText(false)
@@ -27,10 +30,17 @@ public:
 		SLATE_EVENT(FOnIsOpenChanged, OnMenuOpenChanged)
 		SLATE_ARGUMENT(bool, HasDownArrow)
 		SLATE_ARGUMENT(bool, UsesSmallText)
-
 	SLATE_END_ARGS()
 
-	SSimpleComboButton() {}
+	SSimpleComboButton() = default;
 
 	void Construct(const FArguments& InArgs);
-};
\ No newline at end of file
+
+	//~ Begin IActionButton
+	virtual void SetMenuContentWidgetToFocus(TWeakPtr<SWidget> InWidget) override;
+	virtual void SetIsMenuOpen(bool bInIsOpen, bool bInIsFocused) override;
+	//~ End IActionButton
+
+private:
+	TSharedPtr<SActionButton> ActionButton;
+};
diff --git a/Engine/Source/Developer/ToolWidgets/Public/ToolWidgetsSlateTypes.h b/Engine/Source/Developer/ToolWidgets/Public/ToolWidgetsSlateTypes.h
new file mode 100644
--- /dev/null
+++ b/Engine/Source/Developer/ToolWidgets/Public/ToolWidgetsSlateTypes.h
@@ -0,0 +1,146 @@
+// Copyright Epic Games, Inc. All Rights Reserved.
+
+#pragma once
+
+#include "Styling/SlateTypes.h"
+#include "Styling/SlateWidgetStyle.h"
+#include "UObject/Object.h"
+
+#include "ToolWidgetsSlateTypes.generated.h"
+
+enum class EActionButtonType : uint8
+{
+	Default,		// A default, grey button
+	Simple,			// A button with a left-aligned label, and no background or border when unhovered
+	Primary,		// A primary blue button with white text, use for buttons when it's the next logical step (ie. "Next", "OK")
+	Positive,		// A grey button with a green icon indicating a positive action, like "Add"
+	Warning,		// A grey button with a yellow icon indicating a negative action, like "Mute"
+	Error,			// A grey button with a red icon indicating a negative action, like "Cancel"
+
+	Num,
+};
+
+/**
+ * Represents the appearance of an SActionButton
+ */
+USTRUCT(BlueprintType)
+struct FActionButtonStyle : public FSlateWidgetStyle
+{
+	GENERATED_BODY()
+
+	TOOLWIDGETS_API FActionButtonStyle();
+
+	TOOLWIDGETS_API virtual void GetResources(TArray<const FSlateBrush*>& OutBrushes) const override;
+
+	static TOOLWIDGETS_API const FName TypeName;
+	virtual const FName GetTypeName() const override;;
+
+	static TOOLWIDGETS_API const FActionButtonStyle& GetDefault();
+
+	/**
+	 * The type to use for our SActionButton.
+	 */
+	EActionButtonType GetActionButtonType() const;
+	FActionButtonStyle& SetActionButtonType(const EActionButtonType InActionButtonType);
+
+	/**
+	 * The style to use for our SButton.
+	 */
+	UPROPERTY(EditAnywhere, BlueprintReadWrite, Category = "Appearance")
+	FButtonStyle ButtonStyle;
+	FActionButtonStyle& SetButtonStyle(const FButtonStyle& InButtonStyle);
+
+	/**
+	 * The style to use for our SButton when an icon is present. ButtonStyle used if not specified.
+	 */
+	const FButtonStyle& GetIconButtonStyle() const;
+	FActionButtonStyle& SetIconButtonStyle(const FButtonStyle& InButtonStyle);
+
+	/**
+	 * Spacing between button's border and the content. Default uses ButtonStyle.
+	 */
+	UPROPERTY(EditAnywhere, BlueprintReadWrite, Category = "Appearance")
+	TOptional<FMargin> ButtonContentPadding;
+	FActionButtonStyle& SetButtonContentPadding(const FMargin& InContentPadding);
+
+	/**
+	 * The style to use for our SComboButton.
+	 */
+	UPROPERTY(EditAnywhere, BlueprintReadWrite, Category = "Appearance")
+	FComboButtonStyle ComboButtonStyle;
+	FActionButtonStyle& SetComboButtonStyle(const FComboButtonStyle& InComboButtonStyle);
+
+	/**
+	 * Whether to show a down arrow for the combo button
+	 */
+	UPROPERTY(EditAnywhere, BlueprintReadWrite, Category = "Appearance")
+	bool bHasDownArrow;
+	FActionButtonStyle& SetHasDownArrow(const bool& bInHasDownArrow);
+
+	/**
+	 * Spacing between button's border and the content. Default uses ComboButtonStyle.
+	 */
+	UPROPERTY(EditAnywhere, BlueprintReadWrite, Category = "Appearance")
+	TOptional<FMargin> ComboButtonContentPadding;
+	FActionButtonStyle& SetComboButtonContentPadding(const FMargin& InContentPadding);
+
+	/**
+	 * Horizontal Content alignment within the button.
+	 */
+	UPROPERTY(EditAnywhere, BlueprintReadWrite, Category = "Appearance")
+	TEnumAsByte<EHorizontalAlignment> HorizontalContentAlignment;
+	FActionButtonStyle& SetHorizontalContentAlignment(const EHorizontalAlignment& InAlignment);
+
+	/**
+	 * The style to use for the button Text.
+	 */
+	UPROPERTY(EditAnywhere, BlueprintReadWrite, Category = "Appearance")
+	FTextBlockStyle TextBlockStyle;
+	FActionButtonStyle& SetTextBlockStyle(const FTextBlockStyle& InTextBlockStyle);
+
+	/**
+	 * Icon Brush to use.
+	 */
+	UPROPERTY(EditAnywhere, BlueprintReadWrite, Category=Appearance)
+	TOptional<FSlateBrush> IconBrush;
+	FActionButtonStyle& SetIconBrush(const FSlateBrush& InIconBrush);
+
+	/**
+	 * Icon Color/Tint, defaults is determined by ActionButtonType.
+	 */
+	UPROPERTY(EditAnywhere, BlueprintReadWrite, Category = "Appearance")
+	TOptional<FSlateColor> IconColorAndOpacity;
+	FActionButtonStyle& SetIconColorAndOpacity(const FSlateColor& InIconColorAndOpacity);
+
+	/**
+	* Unlinks all colors in this style.
+	* @see FSlateColor::Unlink
+	 */
+	void UnlinkColors();
+
+private:
+	/**
+	 * The type to use for our SActionButton.
+	 */
+	UPROPERTY(EditAnywhere, BlueprintReadWrite, Category = "Appearance", meta = (AllowPrivateAccess = "true", GetOptions = "Engine.ToolSlateWidgetTypesFunctionLibrary.GetActionButtonTypeNames"))
+	FName ActionButtonType;
+
+	/**
+	 * The style to use for our SButton when an icon is present. ButtonStyle used if not specified.
+	 */
+	UPROPERTY(EditAnywhere, BlueprintReadWrite, Category = "Appearance", meta = (AllowPrivateAccess = "true"))
+	TOptional<FButtonStyle> IconButtonStyle;
+};
+
+UCLASS(MinimalAPI)
+class UToolSlateWidgetTypesFunctionLibrary
+	: public UObject
+{
+	GENERATED_BODY()
+
+public:
+#if WITH_EDITOR
+	UFUNCTION()
+	static const TArray<FName>& GetActionButtonTypeNames();
+#endif
+};
diff --git a/Engine/Source/Developer/ToolWidgets/Public/ToolWidgetsStyle.h b/Engine/Source/Developer/ToolWidgets/Public/ToolWidgetsStyle.h
new file mode 100644
--- /dev/null
+++ b/Engine/Source/Developer/ToolWidgets/Public/ToolWidgetsStyle.h
@@ -0,0 +1,37 @@
+// Copyright Epic Games, Inc. All Rights Reserved.
+
+#pragma once
+
+#include "Styling/SlateStyle.h"
+
+namespace UE::ToolWidgets
+{
+	class FToolWidgetsStyle final
+		: public FSlateStyleSet
+	{
+	public:
+		static FName StyleName;
+
+		static TOOLWIDGETS_API FToolWidgetsStyle& Get();
+
+	public:
+		struct FButton
+		{
+			static constexpr float DefaultBorderPadding = 2.0f; // This needs to be accounted for when padding buttons to achieve the desired height, may be overridden by ButtonStyle
+			static constexpr float DefaultBorderThickness = 1.0f;
+		};
+
+		struct FActionButton
+		{
+			static constexpr float DefaultButtonHeight = 22.0f; // Design standard
+			static constexpr float DefaultIconHeight = 16.0f; // Vertical icon space should always be used to affect button size and ensure it's consistent regardless of icon presence
+			static constexpr float DefaultIconLabelSpacing = 3.0f;
+			static constexpr float DefaultVerticalPadding = (DefaultButtonHeight - DefaultIconHeight) / 2.0f; // Total Button Height - Icon Height (/2 for top and bottom)
+			static constexpr float DefaultHorizontalPadding = 2.0f;
+		};
+
+	private:
+		FToolWidgetsStyle();
+		virtual ~FToolWidgetsStyle() override;
+	};
+}
