[ToolWidgets] Made style constants private

diff --git a/Engine/Source/Developer/ToolWidgets/Private/SActionButton.cpp b/Engine/Source/Developer/ToolWidgets/Private/SActionButton.cpp
--- a/Engine/Source/Developer/ToolWidgets/Private/SActionButton.cpp
+++ b/Engine/Source/Developer/ToolWidgets/Private/SActionButton.cpp
@@ -4,6 +4,7 @@
 
 #include "Styling/StyleColors.h"
 #include "ToolWidgetsStyle.h"
+#include "ToolWidgetsStylePrivate.h"
 #include "Widgets/Images/SImage.h"
 #include "Widgets/Input/SButton.h"
 #include "Widgets/Layout/SSpacer.h"
@@ -83,8 +84,8 @@ void SActionButton::Construct(const FArguments& InArgs)
 
 	TAttribute<FText> Text = InArgs._Text;
 
-	static constexpr float DefaultIconHeight = UE::ToolWidgets::FToolWidgetsStyle::FActionButton::DefaultIconHeight;
-	static constexpr float IconTextPadding = UE::ToolWidgets::FToolWidgetsStyle::FActionButton::DefaultIconLabelSpacing;
+	static constexpr float DefaultIconHeight = UE::ToolWidgets::Private::FToolWidgetsStylePrivate::FActionButton::DefaultIconHeight;
+	static constexpr float IconTextPadding = UE::ToolWidgets::Private::FToolWidgetsStylePrivate::FActionButton::DefaultIconLabelSpacing;
 
 	const TSharedRef<SHorizontalBox> ButtonContentContainer = SNew(SHorizontalBox)
 		+ SHorizontalBox::Slot()
@@ -133,8 +134,8 @@ void SActionButton::Construct(const FArguments& InArgs)
 	if (bIsComboButton)
 	{
 		static const FMargin DefaultComboButtonContentPadding = FMargin(
-			UE::ToolWidgets::FToolWidgetsStyle::FActionButton::DefaultHorizontalPadding,
-			UE::ToolWidgets::FToolWidgetsStyle::FActionButton::DefaultVerticalPadding);
+			UE::ToolWidgets::Private::FToolWidgetsStylePrivate::FActionButton::DefaultHorizontalPadding,
+			UE::ToolWidgets::Private::FToolWidgetsStylePrivate::FActionButton::DefaultVerticalPadding);
 
 		const TAttribute<FMargin> ComboButtonContentPadding = InArgs._ButtonContentPadding.IsSet()
 			? InArgs._ButtonContentPadding.Get(DefaultComboButtonContentPadding)
@@ -169,8 +170,8 @@ void SActionButton::Construct(const FArguments& InArgs)
 	else
 	{
 		static const FMargin DefaultButtonContentPadding = FMargin(
-			UE::ToolWidgets::FToolWidgetsStyle::FActionButton::DefaultHorizontalPadding,
-			UE::ToolWidgets::FToolWidgetsStyle::FActionButton::DefaultVerticalPadding);
+			UE::ToolWidgets::Private::FToolWidgetsStylePrivate::FActionButton::DefaultHorizontalPadding,
+			UE::ToolWidgets::Private::FToolWidgetsStylePrivate::FActionButton::DefaultVerticalPadding);
 
 		const TAttribute<FMargin> ButtonContentPadding = InArgs._ButtonContentPadding.IsSet()
 			? InArgs._ButtonContentPadding.Get(DefaultButtonContentPadding)
diff --git a/Engine/Source/Developer/ToolWidgets/Private/ToolWidgetsStyle.cpp b/Engine/Source/Developer/ToolWidgets/Private/ToolWidgetsStyle.cpp
--- a/Engine/Source/Developer/ToolWidgets/Private/ToolWidgetsStyle.cpp
+++ b/Engine/Source/Developer/ToolWidgets/Private/ToolWidgetsStyle.cpp
@@ -6,6 +6,7 @@
 #include "Styling/SlateStyleRegistry.h"
 #include "Styling/StyleColors.h"
 #include "ToolWidgetsSlateTypes.h"
+#include "ToolWidgetsStylePrivate.h"
 
 namespace UE::ToolWidgets
 {
@@ -44,8 +45,10 @@ namespace UE::ToolWidgets
 		// SActionButton
 		{
 			using namespace Private;
-		
-			static const FMargin DefaultButtonContentPadding = FMargin(FActionButton::DefaultHorizontalPadding, FActionButton::DefaultVerticalPadding);
+
+			static const FMargin DefaultButtonContentPadding = FMargin(
+				FToolWidgetsStylePrivate::FActionButton::DefaultHorizontalPadding,
+				FToolWidgetsStylePrivate::FActionButton::DefaultVerticalPadding);
 
 			const FButtonStyle& DefaultButtonStyle = GetWidgetStyle<FButtonStyle>("Button");
 			const FComboButtonStyle& DefaultComboButtonStyle = GetWidgetStyle<FComboButtonStyle>("ComboButton");
diff --git a/Engine/Source/Developer/ToolWidgets/Private/ToolWidgetsStylePrivate.h b/Engine/Source/Developer/ToolWidgets/Private/ToolWidgetsStylePrivate.h
new file mode 100644
--- /dev/null
+++ b/Engine/Source/Developer/ToolWidgets/Private/ToolWidgetsStylePrivate.h
@@ -0,0 +1,28 @@
+// Copyright Epic Games, Inc. All Rights Reserved.
+
+#pragma once
+
+#include "Styling/SlateStyle.h"
+
+namespace UE::ToolWidgets::Private
+{
+	/** A supplementary Style container, accessible to the module only. */
+	class FToolWidgetsStylePrivate final
+	{
+	public:
+		struct FButton
+		{
+			static constexpr float DefaultBorderPadding = 2.0f; // This needs to be accounted for when padding buttons to achieve the desired height, may be overridden by ButtonStyle
+			static constexpr float DefaultBorderThickness = 1.0f;
+		};
+
+		struct FActionButton
+		{
+			static constexpr float DefaultButtonHeight = 22.0f; // Design standard
+			static constexpr float DefaultIconHeight = 16.0f; // Vertical icon space should always be used to affect button size and ensure it's consistent regardless of icon presence
+			static constexpr float DefaultIconLabelSpacing = 3.0f;
+			static constexpr float DefaultVerticalPadding = (DefaultButtonHeight - DefaultIconHeight) / 2.0f; // Total Button Height - Icon Height (/2 for top and bottom)
+			static constexpr float DefaultHorizontalPadding = 2.0f;
+		};
+	};
+}
diff --git a/Engine/Source/Developer/ToolWidgets/Public/ToolWidgetsStyle.h b/Engine/Source/Developer/ToolWidgets/Public/ToolWidgetsStyle.h
--- a/Engine/Source/Developer/ToolWidgets/Public/ToolWidgetsStyle.h
+++ b/Engine/Source/Developer/ToolWidgets/Public/ToolWidgetsStyle.h
@@ -14,22 +14,6 @@ namespace UE::ToolWidgets
 
 		static TOOLWIDGETS_API FToolWidgetsStyle& Get();
 
-	public:
-		struct FButton
-		{
-			static constexpr float DefaultBorderPadding = 2.0f; // This needs to be accounted for when padding buttons to achieve the desired height, may be overridden by ButtonStyle
-			static constexpr float DefaultBorderThickness = 1.0f;
-		};
-
-		struct FActionButton
-		{
-			static constexpr float DefaultButtonHeight = 22.0f; // Design standard
-			static constexpr float DefaultIconHeight = 16.0f; // Vertical icon space should always be used to affect button size and ensure it's consistent regardless of icon presence
-			static constexpr float DefaultIconLabelSpacing = 3.0f;
-			static constexpr float DefaultVerticalPadding = (DefaultButtonHeight - DefaultIconHeight) / 2.0f; // Total Button Height - Icon Height (/2 for top and bottom)
-			static constexpr float DefaultHorizontalPadding = 2.0f;
-		};
-
 	private:
 		FToolWidgetsStyle();
 		virtual ~FToolWidgetsStyle() override;
