[ToolWidgets]
 - Adds SSidebar and related classes for use with Motion Design and Sequencer

diff --git a/Engine/Source/Developer/ToolWidgets/Private/Sidebar/SSidebar.cpp b/Engine/Source/Developer/ToolWidgets/Private/Sidebar/SSidebar.cpp
new file mode 100644
--- /dev/null
+++ b/Engine/Source/Developer/ToolWidgets/Private/Sidebar/SSidebar.cpp
@@ -0,0 +1,879 @@
+// Copyright Epic Games, Inc. All Rights Reserved.
+
+#include "Sidebar/SSidebar.h"
+#include "Framework/Application/SlateApplication.h"
+#include "SidebarButtonMenuContext.h"
+#include "Sidebar/ISidebarDrawerContent.h"
+#include "Sidebar/SSidebarButton.h"
+#include "Sidebar/SSidebarDrawer.h"
+#include "Sidebar/SSidebarDrawerContent.h"
+#include "ToolMenu.h"
+#include "ToolMenus.h"
+#include "Widgets/Layout/SBorder.h"
+#include "Widgets/Layout/SBox.h"
+#include "Widgets/Layout/SScrollBox.h"
+#include "Widgets/SOverlay.h"
+
+#define LOCTEXT_NAMESPACE "SSidebar"
+
+SSidebar::~SSidebar()
+{
+	RemoveAllDrawers();
+}
+
+void SSidebar::Construct(const FArguments& InArgs, const TSharedRef<SOverlay>& InDrawersOverlay, const TSharedRef<SBox>& InDockLocation)
+{
+	DrawersOverlayWeak = InDrawersOverlay;
+	DockLocationWeak = InDockLocation;
+
+	TabLocation = InArgs._TabLocation;
+	bHideWhenDocked = InArgs._HideWhenDocked;
+	bAlwaysUseMaxButtonSize = InArgs._AlwaysUseMaxButtonSize;
+	bDisablePin = InArgs._DisablePin;
+	bDisableDock = InArgs._DisableDock;
+	OnDockStateChanged = InArgs._OnDockStateChanged;
+
+	SetVisibility(EVisibility::SelfHitTestInvisible);
+	
+	ChildSlot
+	.Padding(FMargin(
+		TabLocation == ESidebarTabLocation::Right ? 2.f : 0.f,
+		TabLocation == ESidebarTabLocation::Bottom ? 2.f : 0.f,
+		TabLocation == ESidebarTabLocation::Left ? 2.f : 0.f,
+		TabLocation == ESidebarTabLocation::Top ? 2.f : 0.f))
+	[
+		SNew(SBorder)
+		.Padding(0.0f)
+		.BorderImage(FAppStyle::Get().GetBrush(TEXT("Docking.Sidebar.Background")))
+		[
+			SAssignNew(TabContainer, SScrollBox)
+			.Orientation(IsHorizontal() ? EOrientation::Orient_Horizontal : EOrientation::Orient_Vertical)
+			.ScrollBarAlwaysVisible(false)
+			.ScrollBarVisibility(EVisibility::Collapsed)
+		]
+	];
+}
+
+bool SSidebar::RegisterDrawer(FSidebarDrawerConfig&& InDrawerConfig)
+{
+	if (ContainsDrawer(InDrawerConfig.UniqueId))
+	{
+		return false;
+	}
+
+	const TSharedRef<FSidebarDrawer> NewDrawer = MakeShared<FSidebarDrawer>(MoveTemp(InDrawerConfig));
+	NewDrawer->bDisablePin = bDisablePin;
+	NewDrawer->bDisableDock = bDisableDock;
+	NewDrawer->ContentWidget = NewDrawer->Config.OverrideContentWidget.IsValid()
+		? NewDrawer->ContentWidget = NewDrawer->Config.OverrideContentWidget
+		: NewDrawer->ContentWidget = SNew(SSidebarDrawerContent, NewDrawer);
+
+	// Add tab button
+	TabContainer->AddSlot()
+		[
+			SAssignNew(NewDrawer->ButtonWidget, SSidebarButton, NewDrawer, TabLocation)
+			.MinButtonSize(bAlwaysUseMaxButtonSize ? MaxTabButtonSize : MinTabButtonSize)
+			.MaxButtonSize(MaxTabButtonSize)
+			.ButtonThickness(TabButtonThickness)
+			.OnPressed(this, &SSidebar::OnTabDrawerButtonPressed)
+			.OnPinToggled(this, &SSidebar::OnDrawerTabPinToggled)
+			.OnDockToggled(this, &SSidebar::OnDrawerTabDockToggled)
+			.OnGetContextMenuContent(this, &SSidebar::OnGetTabDrawerContextMenuWidget, NewDrawer)
+		];
+
+	DrawerTabs.Add(NewDrawer);
+
+	// Figure out the size this tab should be when opened later. We do it now when the tab still has valid geometry. Once it is moved to the sidebar it will not.
+	float TargetDrawerSizePct = NewDrawer->SizeCoefficient;
+	if (TargetDrawerSizePct == 0)
+	{
+		TSharedPtr<SWindow> MyWindow = FSlateApplication::Get().FindWidgetWindow(AsShared());
+		if (MyWindow.IsValid() && NewDrawer->ContentWidget.IsValid())
+		{
+			TargetDrawerSizePct = NewDrawer->ContentWidget->GetTickSpaceGeometry().GetLocalSize().X / MyWindow->GetPaintSpaceGeometry().GetLocalSize().X;
+			NewDrawer->SizeCoefficient = TargetDrawerSizePct;
+		}
+	}
+
+	// We don't currently allow more than one pinned tab per sidebar, so enforce that
+	// Note: it's possible to relax this if users actually want multiple pinned tabs
+	if (FindFirstPinnedTab())
+	{
+		SetDrawerPinned(NewDrawer->GetUniqueId(), false);
+	}
+
+	if (NewDrawer->bIsPinned)
+	{
+		// If this tab is a pinned tab, then open the drawer automatically after it's added
+		OpenDrawerNextFrame(NewDrawer, /*bAnimateOpen=*/false);
+	}
+	else if (NewDrawer->Config.bInitiallyDocked)
+	{
+		SetDrawerDocked(NewDrawer->GetUniqueId(), true);
+	}
+	
+	return true;
+}
+
+bool SSidebar::UnregisterDrawer(const FName InDrawerId)
+{
+	if (IsDrawerOpened(InDrawerId))
+	{
+		CloseAllDrawers();
+	}
+	
+	const int32 FoundIndex = DrawerTabs.IndexOfByPredicate(
+		[InDrawerId](const TSharedRef<FSidebarDrawer>& InDrawer)
+		{
+			return InDrawerId == InDrawer->GetUniqueId();
+		});
+	if (FoundIndex == INDEX_NONE)
+	{
+		return false;
+	}
+
+	DrawerTabs.RemoveAt(FoundIndex);
+	TabContainer->RemoveSlot(DrawerTabs[FoundIndex]->ButtonWidget.ToSharedRef());
+
+	RemoveDrawer(DrawerTabs[FoundIndex]);
+	SummonPinnedTabIfNothingOpened();
+
+	// Clear the pinned flag when the tab is removed from the sidebar.
+	// (Users probably expect that pinning a tab, restoring it/closing it,
+	// then moving it to the sidebar again will leave it unpinned the second time.)
+	SetDrawerPinned(InDrawerId, false);
+
+	if (DrawerTabs.Num() == 0)
+	{
+		SetVisibility(EVisibility::Collapsed);
+	}
+
+	return true;
+}
+
+bool SSidebar::ContainsDrawer(const FName InDrawerId) const
+{
+	return FindDrawer(InDrawerId).IsValid();
+}
+
+int32 SSidebar::GetDrawerCount() const
+{
+	return DrawerTabs.Num();
+}
+
+bool SSidebar::RegisterDrawerSection(const FName InDrawerId, const TSharedPtr<ISidebarDrawerContent>& InSection)
+{
+	const TSharedPtr<FSidebarDrawer> Drawer = FindDrawer(InDrawerId);
+	if (!Drawer.IsValid())
+	{
+		return false;
+	}
+
+	const FName SectionUniqueId = InSection->GetUniqueId();
+	if (Drawer->ContentSections.Contains(SectionUniqueId))
+	{
+		return false;
+	}
+
+	Drawer->ContentSections.Add(SectionUniqueId, InSection.ToSharedRef());
+
+	const TSharedPtr<SSidebarDrawerContent> DrawerSection = StaticCastSharedPtr<SSidebarDrawerContent>(Drawer->ContentWidget);
+	if (DrawerSection.IsValid())
+	{
+		DrawerSection->BuildContent();
+	}
+
+	return false;
+}
+
+bool SSidebar::UnregisterDrawerSection(const FName InDrawerId, const FName InSectionId)
+{
+	const TSharedPtr<FSidebarDrawer> Drawer = FindDrawer(InDrawerId);
+	if (!Drawer.IsValid())
+	{
+		return false;
+	}
+
+	if (!Drawer->ContentSections.Contains(InSectionId))
+	{
+		return false;
+	}
+
+	Drawer->ContentSections.Remove(InSectionId);
+
+	return false;
+}
+
+bool SSidebar::TryOpenDrawer(const FName InDrawerId)
+{
+	if (IsDrawerOpened(InDrawerId) || IsDrawerDocked(InDrawerId))
+	{
+		return false;
+	}
+	
+	if (const TSharedPtr<FSidebarDrawer> InDrawer = FindDrawer(InDrawerId))
+	{
+		OpenDrawerNextFrame(InDrawer.ToSharedRef(), /*bAnimateOpen=*/ true);
+		return true;
+	}
+	
+	return false;
+}
+
+void SSidebar::CloseAllDrawers()
+{
+	for (const TSharedRef<FSidebarDrawer>& DrawerTab : DrawerTabs)
+	{
+		CloseDrawerInternal(DrawerTab);
+	}
+}
+
+void SSidebar::OnTabDrawerButtonPressed(const TSharedRef<FSidebarDrawer>& InDrawer)
+{
+	if (InDrawer->bIsOpen)
+	{
+		// When clicking on the button of an active (but unpinned) tab, close that tab drawer
+		if (!IsDrawerPinned(InDrawer->GetUniqueId()))
+		{
+			CloseDrawerInternal(InDrawer);
+		}
+	}
+	else if (!InDrawer->bIsDocked)
+	{
+		// Otherwise clicking on an inactive tab should open the drawer
+		OpenDrawerInternal(InDrawer, /*bAnimateOpen=*/true);
+	}
+}
+
+void SSidebar::OnDrawerTabPinToggled(const TSharedRef<FSidebarDrawer>& InDrawer, const bool bIsPinned)
+{
+	// Set pin state for given tab; clear the pin state for all other tabs
+	for (const TSharedRef<FSidebarDrawer>& DrawerTab : DrawerTabs)
+	{
+		SetDrawerPinned(DrawerTab->GetUniqueId(), DrawerTab == InDrawer ? bIsPinned : false);
+	}
+
+	// Open any newly-pinned tab
+	if (bIsPinned)
+	{
+		OpenDrawerInternal(InDrawer, /*bAnimateOpen=*/true);
+	}
+}
+
+void SSidebar::OnDrawerTabDockToggled(const TSharedRef<FSidebarDrawer>& InDrawer, const bool bIsDocked)
+{
+	// Undock the previously docked drawer
+	if (DockedDrawerTab.IsValid())
+	{
+		SetDrawerDocked(DockedDrawerTab->GetUniqueId(), false);
+	}
+
+	// Dock new drawer if needed
+	if (bIsDocked)
+	{
+		SetDrawerDocked(InDrawer->GetUniqueId(), bIsDocked);
+	}
+}
+
+void SSidebar::OnTabDrawerFocusLost(const TSharedRef<SSidebarDrawer>& InDrawerWidget)
+{
+	const TSharedPtr<FSidebarDrawer> DrawerWidget = InDrawerWidget->GetDrawer();
+	if (!DrawerWidget.IsValid())
+	{
+		return;
+	}
+
+	// Don't automatically close a pinned tab that is in the foreground
+	if (IsDrawerPinned(DrawerWidget->GetUniqueId()) && DrawerWidget == GetForegroundTab())
+	{
+		return;
+	}
+
+	CloseDrawerInternal(DrawerWidget.ToSharedRef());
+}
+
+void SSidebar::OnTabDrawerClosed(const TSharedRef<SSidebarDrawer>& InDrawerWidget)
+{
+	RemoveDrawer(InDrawerWidget->GetDrawer().ToSharedRef());
+}
+
+void SSidebar::OnDrawerTargetSizeChanged(const TSharedRef<SSidebarDrawer>& InDrawerWidget, const float InNewSize)
+{
+	const TSharedPtr<SOverlay> DrawersOverlay = DrawersOverlayWeak.Pin();
+	if (!DrawersOverlay.IsValid())
+	{
+		return;
+	}
+	
+	const TSharedPtr<FSidebarDrawer> DrawerWidget = InDrawerWidget->GetDrawer();
+	if (!DrawerWidget.IsValid())
+	{
+		return;
+	}
+	
+	DrawerWidget->SizeCoefficient = InNewSize / DrawersOverlay->GetPaintSpaceGeometry().GetLocalSize().X;
+}
+
+TSharedRef<SWidget> SSidebar::OnGetTabDrawerContextMenuWidget(TSharedRef<FSidebarDrawer> InDrawer)
+{
+	UToolMenus* const ToolMenus = UToolMenus::Get();
+	if (!IsValid(ToolMenus))
+	{
+		return SNullWidget::NullWidget;
+	}
+
+	static constexpr const TCHAR* MenuName = TEXT("SidebarTabMenu");
+
+	if (!ToolMenus->IsMenuRegistered(MenuName))
+	{
+		UToolMenu* const NewMenu = ToolMenus->RegisterMenu(MenuName, NAME_None, EMultiBoxType::Menu);
+		check(IsValid(NewMenu));
+		
+		NewMenu->AddDynamicSection(TEXT("Options"), FNewToolMenuDelegate::CreateSP(this, &SSidebar::BuildOptionsMenu));
+	}
+
+	USidebarButtonMenuContext* const ContextObject = NewObject<USidebarButtonMenuContext>();
+	ContextObject->Init(SharedThis(this), InDrawer);
+
+	const FToolMenuContext MenuContext(nullptr, nullptr, ContextObject);
+	return ToolMenus->GenerateWidget(MenuName, MenuContext);
+}
+
+void SSidebar::BuildOptionsMenu(UToolMenu* const InMenu)
+{
+	if (!IsValid(InMenu))
+	{
+		return;
+	}
+
+	USidebarButtonMenuContext* const ContextMenu = InMenu->FindContext<USidebarButtonMenuContext>();
+	if (!IsValid(ContextMenu))
+	{
+		return;
+	}
+
+	const TSharedPtr<FSidebarDrawer> Drawer = ContextMenu->GetDrawer();
+	if (!Drawer.IsValid())
+	{
+		return;
+	}
+
+	FToolMenuSection& Section = InMenu->FindOrAddSection(TEXT("Options"), LOCTEXT("Options", "Options"));
+
+	if (Drawer->bIsDocked)
+	{
+		Section.AddMenuEntry(TEXT("Undock"),
+			LOCTEXT("UndockLabel", "Undock"),
+			LOCTEXT("UndockToolTip", "Undocks the drawer"),
+			FSlateIcon(),
+			FUIAction(FExecuteAction::CreateSP(this, &SSidebar::SetDrawerDocked, Drawer->GetUniqueId(), false)));
+	}
+	else
+	{
+		Section.AddMenuEntry(TEXT("Dock"),
+			LOCTEXT("UndockLabel", "Dock"),
+			LOCTEXT("UndockToolTip", "Docks the drawer"),
+			FSlateIcon(),
+			FUIAction(FExecuteAction::CreateSP(this, &SSidebar::SetDrawerDocked, Drawer->GetUniqueId(), true)));
+	}
+
+	if (Drawer->bIsPinned)
+	{
+		Section.AddMenuEntry(TEXT("Unpin"),
+			LOCTEXT("UnpinLabel", "Unpin"),
+			LOCTEXT("UnpinTooltip", "Unpins the drawer from always being displayed"),
+			FSlateIcon(),
+			FUIAction(FExecuteAction::CreateSP(this, &SSidebar::SetDrawerPinned, Drawer->GetUniqueId(), false)));
+	}
+	else
+	{
+		Section.AddMenuEntry(TEXT("Pin"),
+			LOCTEXT("PinLabel", "Pin"),
+			LOCTEXT("PinTooltip", "Pins the drawer to always be displayed"),
+			FSlateIcon(),
+			FUIAction(FExecuteAction::CreateSP(this, &SSidebar::SetDrawerPinned, Drawer->GetUniqueId(), true)));
+	}
+}
+
+void SSidebar::RemoveDrawer(const TSharedRef<FSidebarDrawer>& InDrawer)
+{
+	if (InDrawer->DrawerWidget.IsValid())
+	{
+		if (const TSharedPtr<SOverlay> DrawersOverlay = DrawersOverlayWeak.Pin())
+		{
+			DrawersOverlay->RemoveSlot(InDrawer->DrawerWidget.ToSharedRef());
+		}
+	}
+
+	InDrawer->bIsOpen = false;
+
+	InDrawer->DrawerClosedDelegate.ExecuteIfBound(InDrawer->GetUniqueId());
+
+	UpdateDrawerAppearance();
+}
+
+void SSidebar::RemoveAllDrawers()
+{
+	PendingTabToOpen.Reset();
+	bAnimatePendingTabOpen = false;
+
+	// Closing drawers can remove them from the opened drawers list so copy the list first
+	TArray<TSharedRef<SSidebarDrawer>> OpenedDrawersCopy = OpenedDrawers;
+
+	for (TSharedRef<SSidebarDrawer>& Drawer : OpenedDrawersCopy)
+	{
+		RemoveDrawer(Drawer->GetDrawer().ToSharedRef());
+	}
+}
+
+EActiveTimerReturnType SSidebar::OnOpenPendingDrawerTimer(const double CurrentTime, const float DeltaTime)
+{
+	if (const TSharedPtr<FSidebarDrawer> TabToOpen = PendingTabToOpen.Pin())
+	{
+		// Wait until the drawers overlay has been arranged once to open the drawer
+		// It might not have geometry yet if we're adding back tabs on startup
+		if (const TSharedPtr<SOverlay> DrawersOverlay = DrawersOverlayWeak.Pin())
+		{
+			if (DrawersOverlay->GetTickSpaceGeometry().GetLocalSize().IsZero())
+			{
+				return EActiveTimerReturnType::Continue;
+			}
+		}
+
+		OpenDrawerInternal(TabToOpen.ToSharedRef(), bAnimatePendingTabOpen);
+	}
+
+	PendingTabToOpen.Reset();
+	bAnimatePendingTabOpen = false;
+	OpenPendingDrawerTimerHandle.Reset();
+
+	return EActiveTimerReturnType::Stop;
+}
+
+void SSidebar::OpenDrawerNextFrame(const TSharedRef<FSidebarDrawer>& InDrawer, const bool bAnimateOpen)
+{
+	PendingTabToOpen = InDrawer;
+	bAnimatePendingTabOpen = bAnimateOpen;
+
+	if (!OpenPendingDrawerTimerHandle.IsValid())
+	{
+		OpenPendingDrawerTimerHandle = RegisterActiveTimer(0.f, FWidgetActiveTimerDelegate::CreateSP(this, &SSidebar::OnOpenPendingDrawerTimer));
+	}
+}
+
+void SSidebar::OpenDrawerInternal(const TSharedRef<FSidebarDrawer>& InDrawer, const bool bAnimateOpen)
+{
+	if (FindOpenedDrawer(InDrawer))
+	{
+		return;
+	}
+
+	const TSharedPtr<SOverlay> DrawersOverlay = DrawersOverlayWeak.Pin();
+	if (!DrawersOverlay.IsValid())
+	{
+		return;
+	}
+
+	PendingTabToOpen.Reset();
+	bAnimatePendingTabOpen = false;
+
+	const FGeometry DrawersOverlayGeometry = DrawersOverlay->GetTickSpaceGeometry();
+	const FGeometry Geometry = GetTickSpaceGeometry();
+
+	// Calculate padding for the drawer itself
+	const float MinDrawerSize = Geometry.GetLocalSize().X - 4.f; // overlap with sidebar border slightly
+	const FVector2D ShadowOffset(8.f, 8.f);
+	const FMargin SlotPadding(
+		TabLocation == ESidebarTabLocation::Left ? MinDrawerSize : 0.f,
+		-ShadowOffset.Y,
+		TabLocation == ESidebarTabLocation::Right ? MinDrawerSize : 0.f,
+		-ShadowOffset.Y);
+	const float AvailableWidth = DrawersOverlayGeometry.GetLocalSize().X - SlotPadding.GetTotalSpaceAlong<EOrientation::Orient_Horizontal>();
+	const float MaxDrawerSize = AvailableWidth * 0.5f;
+
+	float TargetDrawerSizePct = InDrawer->SizeCoefficient;
+	TargetDrawerSizePct = FMath::Clamp(TargetDrawerSizePct, 0.f, 0.5f);
+
+	const float TargetDrawerSize = AvailableWidth * TargetDrawerSizePct;
+
+	InDrawer->DrawerWidget =
+		SNew(SSidebarDrawer, InDrawer, TabLocation)
+		.MinDrawerSize(MinDrawerSize)
+		.TargetDrawerSize(TargetDrawerSize)
+		.MaxDrawerSize(MaxDrawerSize)
+		.OnDrawerFocusLost(this, &SSidebar::OnTabDrawerFocusLost)
+		.OnDrawerClosed(this, &SSidebar::OnTabDrawerClosed)
+		.OnDrawerTargetSizeChanged(this, &SSidebar::OnDrawerTargetSizeChanged);
+
+	DrawersOverlay->AddSlot()
+		.Padding(SlotPadding)
+		.HAlign(TabLocation == ESidebarTabLocation::Left ? HAlign_Left : HAlign_Right)
+		[
+			InDrawer->DrawerWidget.ToSharedRef()
+		];
+
+	InDrawer->DrawerWidget->Open(bAnimateOpen);
+
+	OpenedDrawers.Add(InDrawer->DrawerWidget.ToSharedRef());
+
+	for (const TSharedRef<FSidebarDrawer>& DrawerTab : DrawerTabs)
+	{
+		DrawerTab->bIsOpen = false;
+	}
+	InDrawer->bIsOpen = true;
+
+	InDrawer->DrawerOpenedDelegate.ExecuteIfBound(InDrawer->GetUniqueId());
+
+	UpdateDrawerAppearance();
+
+	// This changes the focus and will trigger focus-related events, such as closing other tabs,
+	// so it's important that we only call it after we added the new drawer to OpenedDrawers.
+	FSlateApplication::Get().SetKeyboardFocus(InDrawer->DrawerWidget);
+}
+
+void SSidebar::CloseDrawerInternal(const TSharedRef<FSidebarDrawer>& InDrawer)
+{
+	if (const TSharedPtr<SSidebarDrawer> OpenedDrawer = FindOpenedDrawer(InDrawer))
+	{
+		OpenedDrawer->Close();
+		const TSharedRef<SSidebarDrawer> OpenedDrawerRef = OpenedDrawer.ToSharedRef();
+
+		if (const TSharedPtr<SOverlay> DrawersOverlay = DrawersOverlayWeak.Pin())
+		{
+			DrawersOverlay->RemoveSlot(OpenedDrawerRef);
+		}
+
+		OpenedDrawers.Remove(OpenedDrawerRef);
+	}
+
+	InDrawer->bIsOpen = false;
+
+	SummonPinnedTabIfNothingOpened();
+	UpdateDrawerAppearance();
+}
+
+void SSidebar::SummonPinnedTabIfNothingOpened()
+{
+	// If there's already a tab in the foreground, don't bring the pinned tab forward
+	if (GetForegroundTab())
+	{
+		return;
+	}
+
+	// But if there's no current foreground tab, then bring forward a pinned tab (there should be at most one)
+	// This should happen when:
+	// - the current foreground tab is not pinned and loses focus
+	// - the current foreground tab's drawer is manually closed by pressing on the tab button
+	// - closing or restoring the current foreground tab
+	if (const TSharedPtr<FSidebarDrawer> PinnedTab = FindFirstPinnedTab())
+	{
+		OpenDrawerInternal(PinnedTab.ToSharedRef(), /*bAnimateOpen=*/true);
+	}
+}
+
+void SSidebar::UpdateDrawerAppearance()
+{
+	TSharedPtr<FSidebarDrawer> OpenedTab;
+	if (OpenedDrawers.Num() > 0)
+	{
+		OpenedTab = OpenedDrawers.Last()->GetDrawer();
+	}
+
+	for (const TSharedRef<FSidebarDrawer>& Drawer : DrawerTabs)
+	{
+		const TSharedPtr<SSidebarButton> TabButton = StaticCastSharedPtr<SSidebarButton>(Drawer->ButtonWidget);
+		if (TabButton.IsValid())
+		{
+			TabButton->UpdateAppearance(OpenedTab);
+		}
+	}
+}
+
+TSharedPtr<FSidebarDrawer> SSidebar::FindDrawer(const FName InDrawerId) const
+{
+	const TSharedRef<FSidebarDrawer>* const FoundDrawer = DrawerTabs.FindByPredicate(
+		[InDrawerId](const TSharedRef<FSidebarDrawer>& InDrawer)
+		{
+			return InDrawerId == InDrawer->GetUniqueId();
+		});
+	return FoundDrawer ? *FoundDrawer : TSharedPtr<FSidebarDrawer>();
+}
+
+TSharedPtr<FSidebarDrawer> SSidebar::FindFirstPinnedTab() const
+{
+	for (const TSharedRef<FSidebarDrawer>& DrawerTab : DrawerTabs)
+	{
+		if (DrawerTab->bIsPinned)
+		{
+			return DrawerTab;
+		}
+	}
+	return nullptr;
+}
+
+TSharedPtr<FSidebarDrawer> SSidebar::GetForegroundTab() const
+{
+	const int32 Index = OpenedDrawers.FindLastByPredicate(
+		[](const TSharedRef<SSidebarDrawer>& Drawer)
+		{
+			return Drawer->IsOpen() && !Drawer->IsClosing();
+		});
+	return Index == INDEX_NONE ? nullptr : OpenedDrawers[Index]->GetDrawer();
+}
+
+TSharedPtr<SSidebarDrawer> SSidebar::FindOpenedDrawer(const TSharedRef<FSidebarDrawer>& InDrawer) const
+{
+	const TSharedRef<SSidebarDrawer>* OpenedDrawer = OpenedDrawers.FindByPredicate(
+		[&InDrawer](const TSharedRef<SSidebarDrawer>& Drawer)
+		{
+			return InDrawer == Drawer->GetDrawer();
+		});
+	return OpenedDrawer ? TSharedPtr<SSidebarDrawer>(*OpenedDrawer) : nullptr;
+}
+
+bool SSidebar::HasDrawerOpened() const
+{
+	for (const TSharedRef<FSidebarDrawer>& DrawerTab : DrawerTabs)
+	{
+		if (DrawerTab->bIsOpen)
+		{
+			return true;
+		}
+	}
+	return false;
+}
+
+bool SSidebar::IsDrawerOpened(const FName InDrawerId) const
+{
+	for (const TSharedRef<FSidebarDrawer>& DrawerTab : DrawerTabs)
+	{
+		if (DrawerTab->bIsOpen && DrawerTab->GetUniqueId() == InDrawerId)
+		{
+			return true;
+		}
+	}
+	return false;
+}
+
+FName SSidebar::GetOpenedDrawerId() const
+{
+	if (OpenedDrawers.IsEmpty())
+	{
+		return NAME_None;
+	}
+	
+	const TSharedRef<SSidebarDrawer> LastOpenedDrawer = OpenedDrawers.Last();
+	return LastOpenedDrawer->GetDrawer()->GetUniqueId();
+}
+
+bool SSidebar::HasDrawerPinned() const
+{
+	for (const TSharedRef<FSidebarDrawer>& DrawerTab : DrawerTabs)
+	{
+		if (DrawerTab->bIsPinned)
+		{
+			return true;
+		}
+	}
+	return false;
+}
+
+bool SSidebar::IsDrawerPinned(const FName InDrawerId) const
+{
+	if (const TSharedPtr<FSidebarDrawer> Drawer = FindDrawer(InDrawerId))
+	{
+		return PinnedDrawerTabs.Contains(Drawer.ToSharedRef());
+	}
+	return false;
+}
+
+void SSidebar::SetDrawerPinned(const FName InDrawerId, const bool bInIsPinned)
+{
+	const TSharedPtr<FSidebarDrawer> Drawer = FindDrawer(InDrawerId);
+	if (!Drawer.IsValid() || Drawer->bIsPinned == bInIsPinned)
+	{
+		return;
+	}
+
+	if (bInIsPinned)
+	{
+		if (Drawer->bIsDocked)
+		{
+			SetDrawerDocked(InDrawerId, false);
+		}
+
+		if (!Drawer->bIsOpen)
+		{
+			OpenDrawerInternal(Drawer.ToSharedRef(), false);
+		}
+		if (!Drawer->bIsOpen)
+		{
+			return;
+		}
+	}
+
+	Drawer->bIsPinned = bInIsPinned;
+
+	if (bInIsPinned)
+	{
+		PinnedDrawerTabs.AddUnique(Drawer.ToSharedRef());
+	}
+	else
+	{
+		PinnedDrawerTabs.Remove(Drawer.ToSharedRef());
+	}
+}
+
+bool SSidebar::HasDrawerDocked() const
+{
+	return DockedDrawerTab.IsValid();
+}
+
+bool SSidebar::IsDrawerDocked(const FName InDrawerId) const
+{
+	if (!DockedDrawerTab.IsValid())
+	{
+		return false;
+	}
+
+	const TSharedPtr<FSidebarDrawer> DrawerConfig = FindDrawer(InDrawerId);
+	if (!DrawerConfig.IsValid())
+	{
+		return false;
+	}
+
+	return DrawerConfig->GetUniqueId() == DockedDrawerTab->GetUniqueId();
+}
+
+void SSidebar::SetDrawerDocked(const FName InDrawerId, const bool bInIsDocked)
+{
+	const TSharedPtr<SBox> DockLocation = DockLocationWeak.Pin();
+	if (!DockLocation.IsValid())
+	{
+		return;
+	}
+
+	const TSharedPtr<FSidebarDrawer> Drawer = FindDrawer(InDrawerId);
+	if (!Drawer.IsValid() || Drawer->bIsDocked == bInIsDocked)
+	{
+		return;
+	}
+
+	if (bInIsDocked)
+	{
+		if (Drawer->bIsPinned)
+		{
+			SetDrawerPinned(InDrawerId, false);
+		}
+
+		CloseAllDrawers();
+
+		if (DockedDrawerTab.IsValid())
+		{
+			UndockAllDrawers();
+		}
+
+		DockedDrawerTab = Drawer;
+
+		DockedDrawerTab->bIsPinned = false;
+		DockedDrawerTab->bIsDocked = true;
+
+		if (Drawer->ContentWidget.IsValid())
+		{
+			DockLocation->SetContent(Drawer->ContentWidget.ToSharedRef());
+			
+			if (bHideWhenDocked)
+			{
+				SetVisibility(EVisibility::Collapsed);
+			}
+		}
+		else
+		{
+			DockLocation->SetContent(SNullWidget::NullWidget);
+			
+			if (bHideWhenDocked)
+			{
+				SetVisibility(EVisibility::Visible);
+			}
+		}
+	}
+	else
+	{
+		if (DockedDrawerTab.IsValid())
+		{
+			DockedDrawerTab->bIsDocked = false;
+			DockedDrawerTab.Reset();
+
+			DockLocation->SetContent(SNullWidget::NullWidget);
+			if (bHideWhenDocked)
+			{
+				SetVisibility(EVisibility::Visible);
+			}
+		}
+	}
+
+	OnDockStateChanged.ExecuteIfBound(InDrawerId);
+}
+
+void SSidebar::UndockAllDrawers()
+{
+	for (const TSharedRef<FSidebarDrawer>& DrawerTab : DrawerTabs)
+	{
+		SetDrawerDocked(DrawerTab->GetUniqueId(), false);
+	}
+}
+
+void SSidebar::UpdateDockedSplitterSlot(const FName InDrawerId, SSplitter::FSlot* const InSlot, const bool bInAutoUndock, const float InDefaultDockPercent)
+{
+	if (!InSlot)
+	{
+		return;
+	}
+	
+	const TSharedPtr<FSidebarDrawer> Drawer = FindDrawer(InDrawerId);
+	if (!Drawer.IsValid())
+	{
+		return;
+	}
+
+	const bool bDocked = IsDrawerDocked(InDrawerId);
+
+	InSlot->SetSizingRule(bDocked ? SSplitter::ESizeRule::FractionOfParent : SSplitter::ESizeRule::SizeToContent);
+	InSlot->SetResizable(bDocked);
+	if (bInAutoUndock && InSlot->GetSizeValue() < 0.01f)
+	{
+		Drawer->SizeCoefficient = bDocked ? InDefaultDockPercent : 0.f;
+		InSlot->SetSizeValue(Drawer->SizeCoefficient);
+	}
+}
+
+bool SSidebar::ContainsDrawerSection(const FName InDrawerId, const FName InDrawerSectionId) const
+{
+	const TSharedPtr<FSidebarDrawer> Drawer = FindDrawer(InDrawerId);
+	if (!Drawer.IsValid())
+	{
+		return false;
+	}
+
+	for (const TPair<FName, TSharedRef<ISidebarDrawerContent>>& DrawerSection : Drawer->ContentSections)
+	{
+		if (DrawerSection.Value->GetSectionId() == InDrawerSectionId)
+		{
+			return true;
+		}
+	}
+
+	return false;
+}
+
+bool SSidebar::IsHorizontal() const
+{
+	return TabLocation == ESidebarTabLocation::Top || TabLocation == ESidebarTabLocation::Bottom;
+}
+
+bool SSidebar::IsVertical() const
+{
+	return TabLocation == ESidebarTabLocation::Left || TabLocation == ESidebarTabLocation::Right;
+}
+
+#undef LOCTEXT_NAMESPACE
diff --git a/Engine/Source/Developer/ToolWidgets/Private/Sidebar/SSidebarButton.cpp b/Engine/Source/Developer/ToolWidgets/Private/Sidebar/SSidebarButton.cpp
new file mode 100644
--- /dev/null
+++ b/Engine/Source/Developer/ToolWidgets/Private/Sidebar/SSidebarButton.cpp
@@ -0,0 +1,422 @@
+// Copyright Epic Games, Inc. All Rights Reserved.
+
+#include "SSidebarButton.h"
+#include "Framework/Application/SlateApplication.h"
+#include "Sidebar/SidebarDrawer.h"
+#include "Sidebar/SSidebar.h"
+#include "Sidebar/SSidebarButtonText.h"
+#include "Styling/AppStyle.h"
+#include "Widgets/Colors/SComplexGradient.h"
+#include "Widgets/Images/SImage.h"
+#include "Widgets/Input/SButton.h"
+#include "Widgets/Input/SCheckBox.h"
+#include "Widgets/Layout/SBorder.h"
+#include "Widgets/Layout/SBox.h"
+#include "Widgets/SOverlay.h"
+
+#define LOCTEXT_NAMESPACE "SSidebarDrawerButton"
+
+void SSidebarButton::Construct(const FArguments& InArgs, const TSharedRef<FSidebarDrawer>& InDrawer, const ESidebarTabLocation InTabLocation)
+{
+	DrawerWeak = InDrawer;
+	TabLocation = InTabLocation;
+
+	OnPressed = InArgs._OnPressed;
+	OnPinToggled = InArgs._OnPinToggled;
+	OnDockToggled = InArgs._OnDockToggled;
+	OnGetContextMenuContent = InArgs._OnGetContextMenuContent;
+
+	DockTabStyle = &FAppStyle::Get().GetWidgetStyle<FDockTabStyle>(TEXT("Docking.Tab"));
+
+	static FLinearColor ActiveBorderColor = FAppStyle::Get().GetSlateColor(TEXT("Docking.Tab.ActiveTabIndicatorColor")).GetSpecifiedColor();
+	static FLinearColor ActiveBorderColorTransparent = FLinearColor(ActiveBorderColor.R, ActiveBorderColor.G, ActiveBorderColor.B, 0.0f);
+	static TArray<FLinearColor> GradientStops{ ActiveBorderColorTransparent, ActiveBorderColor, ActiveBorderColorTransparent };
+	
+	const bool bIsHorizontal = TabLocation == ESidebarTabLocation::Top || TabLocation == ESidebarTabLocation::Bottom;
+	const float MinDesiredWidth = bIsHorizontal ? InArgs._MinButtonSize : InArgs._ButtonThickness;
+	const float MinDesiredHeight = bIsHorizontal ? InArgs._ButtonThickness : InArgs._MinButtonSize;
+	const float MaxDesiredWidth = bIsHorizontal ? InArgs._MaxButtonSize : InArgs._ButtonThickness;
+	const float MaxDesiredHeight = bIsHorizontal ? InArgs._ButtonThickness : InArgs._MaxButtonSize;
+
+	TSharedPtr<SImage> IconWidget;
+	if (InDrawer->Config.Icon.IsSet() || InDrawer->Config.Icon.IsBound())
+	{
+		IconWidget = SNew(SImage)
+			.ColorAndOpacity(FSlateColor::UseForeground())
+			.Image(InDrawer->Config.Icon)
+			.DesiredSizeOverride(FVector2D(16, 16));
+	}
+
+	Label.Reset();
+	if (InDrawer->Config.ButtonText.IsSet() || InDrawer->Config.ButtonText.IsBound())
+	{
+		SAssignNew(Label, SSidebarButtonText)
+			.TextStyle(&DockTabStyle->TabTextStyle)
+			.Text(InDrawer->Config.ButtonText)
+			.OverflowPolicy(ETextOverflowPolicy::Ellipsis)
+			.Clipping(EWidgetClipping::ClipToBounds);
+	}
+
+	SAssignNew(PinCheckBox, SCheckBox)
+		.Style(FAppStyle::Get(), TEXT("ToggleButtonCheckbox"))
+		.Visibility(this, &SSidebarButton::GetPinVisibility)
+		.ToolTipText(this, &SSidebarButton::GetPinToolTipText)
+		.IsChecked(this, &SSidebarButton::IsPinChecked)
+		.OnCheckStateChanged(this, &SSidebarButton::OnPinStateChanged)
+		.Padding(2.0f)
+		.HAlign(HAlign_Center)
+		[
+			SNew(SImage)
+			.ColorAndOpacity(FSlateColor::UseForeground())
+			.Image(this, &SSidebarButton::GetPinImage)
+		];
+
+	SAssignNew(DockCheckBox, SCheckBox)
+		.Style(FAppStyle::Get(), TEXT("ToggleButtonCheckbox"))
+		.Visibility(this, &SSidebarButton::GetDockVisibility)
+		.ToolTipText(this, &SSidebarButton::GetDockToolTipText)
+		.IsChecked(this, &SSidebarButton::IsDockChecked)
+		.OnCheckStateChanged(this, &SSidebarButton::OnDockStateChanged)
+		.Padding(2.0f)
+		.HAlign(HAlign_Center)
+		[
+			SNew(SImage)
+			.ColorAndOpacity(FSlateColor::UseForeground())
+			.Image(this, &SSidebarButton::GetDockImage)
+		];
+
+	TSharedPtr<SWidget> ButtonContent;
+
+	const bool bIsVertical = InTabLocation == ESidebarTabLocation::Left || InTabLocation == ESidebarTabLocation::Right;
+	if (bIsVertical)
+	{
+		const TSharedRef<SVerticalBox> VerticalContent = SNew(SVerticalBox);
+
+		if (IconWidget.IsValid())
+		{
+			VerticalContent->AddSlot()
+				.AutoHeight()
+				.HAlign(HAlign_Center)
+				.VAlign(VAlign_Center)
+				.Padding(0.f, 2.f, 0.f, 3.f)
+				[
+					IconWidget.ToSharedRef()
+				];
+		}
+
+		if (Label.IsValid())
+		{
+			VerticalContent->AddSlot()
+				.FillHeight(1.f)
+				.HAlign(HAlign_Center)
+				.Padding(0.f, 3.f, 0.f, 3.f)
+				[
+					Label.ToSharedRef()
+				];
+		}
+
+		if (!InDrawer->bDisablePin)
+		{
+			VerticalContent->AddSlot()
+				.AutoHeight()
+				.HAlign(HAlign_Center)
+				.Padding(0.f, 3.f, 0.f, 1.f)
+				[
+					PinCheckBox.ToSharedRef()
+				];
+		}
+
+		if (!InDrawer->bDisableDock)
+		{
+			VerticalContent->AddSlot()
+				.AutoHeight()
+				.HAlign(HAlign_Center)
+				.Padding(0.0f, 1.0f, 0.0f, 3.0f)
+				[
+					DockCheckBox.ToSharedRef()
+				];
+		}
+
+		ButtonContent = VerticalContent;
+	}
+	else
+	{
+		const TSharedRef<SHorizontalBox> HorizontalContent = SNew(SHorizontalBox);
+
+		if (IconWidget.IsValid())
+		{
+			HorizontalContent->AddSlot()
+				.AutoWidth()
+				.HAlign(HAlign_Center)
+				.VAlign(VAlign_Center)
+				.Padding(2.f, 0.f, 3.f, 0.f)
+				[
+					IconWidget.ToSharedRef()
+				];
+		}
+		
+		if (Label.IsValid())
+		{
+			HorizontalContent->AddSlot()
+				.FillWidth(1.f)
+				.VAlign(VAlign_Center)
+				.Padding(3.f, 0.f, 3.f, 0.f)
+				[
+					Label.ToSharedRef()
+				];
+		}
+		
+		if (!InDrawer->bDisablePin)
+		{
+			HorizontalContent->AddSlot()
+				.AutoWidth()
+				.VAlign(VAlign_Center)
+				.Padding(3.f, 0.f, 1.f, 0.f)
+				[
+					PinCheckBox.ToSharedRef()
+				];
+		}
+
+		if (!InDrawer->bDisableDock)
+		{
+			HorizontalContent->AddSlot()
+				.AutoWidth()
+				.VAlign(VAlign_Center)
+				.Padding(1.f, 0.f, 3.f, 0.f)
+				[
+					DockCheckBox.ToSharedRef()
+				];
+		}
+
+		ButtonContent = HorizontalContent;
+	}
+	
+	ChildSlot
+	.Padding(0)
+	[
+		SNew(SBox)
+		.MinDesiredWidth(MinDesiredWidth)
+		.MaxDesiredWidth(MaxDesiredWidth)
+		.MinDesiredHeight(MinDesiredHeight)
+		.MaxDesiredHeight(MaxDesiredHeight)
+		.Clipping(EWidgetClipping::ClipToBounds)
+		[
+			SNew(SOverlay)
+			+ SOverlay::Slot()
+			[
+				SAssignNew(MainButton, SButton)
+				.ToolTipText(InDrawer->Config.ToolTipText)
+				.ContentPadding(FMargin(0.0f, DockTabStyle->TabPadding.Top, 0.0f, DockTabStyle->TabPadding.Bottom))
+				.OnPressed_Lambda([this]()
+					{
+						// Activate tab on mouse down (not mouse down-up) for consistency with non-sidebar tabs
+						OnPressed.ExecuteIfBound(DrawerWeak.Pin().ToSharedRef());
+					})
+				.ForegroundColor(FSlateColor::UseForeground())
+				[
+					ButtonContent.ToSharedRef()
+				]
+			]
+			+ SOverlay::Slot()
+			[
+				SAssignNew(OpenBorder, SBorder)
+				.Visibility(EVisibility::HitTestInvisible)
+			]
+			+ SOverlay::Slot()
+			.HAlign(TabLocation == ESidebarTabLocation::Left ? HAlign_Left : HAlign_Right)
+			[
+				SAssignNew(ActiveIndicator, SComplexGradient)
+				.DesiredSizeOverride(FVector2D(1.0f, 1.0f))
+				.GradientColors(GradientStops)
+				.Orientation(EOrientation::Orient_Horizontal)
+				.Visibility(this, &SSidebarButton::GetActiveTabIndicatorVisibility)
+			]
+		]
+	];
+
+	UpdateAppearance(nullptr);
+}
+
+void SSidebarButton::UpdateAppearance(const TSharedPtr<FSidebarDrawer>& InOpenedDrawer)
+{
+	float LabelRotation;
+	FName OpenBorderBrushName;
+
+	switch (TabLocation)
+	{
+	case ESidebarTabLocation::Left:
+		LabelRotation = -90.f;
+		OpenBorderBrushName = TEXT("Docking.Sidebar.Border_SquareRight");
+		break;
+	case ESidebarTabLocation::Right:
+		LabelRotation = 90.f;
+		OpenBorderBrushName = TEXT("Docking.Sidebar.Border_SquareLeft");
+		break;
+	default:
+		LabelRotation = 0.f;
+		OpenBorderBrushName = TEXT("None");
+		break;
+	}
+
+	if (Label.IsValid())
+	{
+		Label->SetRotation(LabelRotation);
+	}
+
+	if (InOpenedDrawer == DrawerWeak)
+	{
+		// this button is the one with the tab that is actually opened so show the tab border
+		OpenBorder->SetVisibility(EVisibility::HitTestInvisible);
+		MainButton->SetButtonStyle(&FAppStyle::Get().GetWidgetStyle<FButtonStyle>(TEXT("Docking.SidebarButton.Opened")));
+
+		OpenBorder->SetBorderImage(FAppStyle::Get().GetBrush(OpenBorderBrushName));
+	}
+	else
+	{
+		OpenBorder->SetVisibility(EVisibility::Collapsed);
+		MainButton->SetButtonStyle(&FAppStyle::Get().GetWidgetStyle<FButtonStyle>(TEXT("Docking.SidebarButton.Closed")));
+	}
+}
+
+void SSidebarButton::OnTabRenamed(const TWeakPtr<FSidebarDrawer>& InDrawer)
+{
+	if (ensure(InDrawer == DrawerWeak) && InDrawer.IsValid())
+	{
+		if (const TSharedPtr<FSidebarDrawer> DrawerConfig = InDrawer.Pin())
+		{
+			if (Label.IsValid())
+			{
+				Label->SetText(DrawerConfig->Config.ButtonText);
+			}
+
+			MainButton->SetToolTipText(DrawerConfig->Config.ToolTipText);
+		}
+	}
+}
+
+FReply SSidebarButton::OnMouseButtonDown(const FGeometry& InGeometry, const FPointerEvent& InMouseEvent)
+{
+	if (InMouseEvent.GetEffectingButton() == EKeys::RightMouseButton && OnGetContextMenuContent.IsBound())
+	{
+		FWidgetPath WidgetPath = InMouseEvent.GetEventPath() != nullptr ? *InMouseEvent.GetEventPath() : FWidgetPath();
+		FSlateApplication::Get().PushMenu(AsShared(), WidgetPath, OnGetContextMenuContent.Execute(), FSlateApplication::Get().GetCursorPos(), FPopupTransitionEffect::ContextMenu);
+		return FReply::Handled();
+	}
+	return FReply::Unhandled();
+}
+
+FSlateColor SSidebarButton::GetForegroundColor() const
+{
+	if (ActiveIndicator->GetVisibility() != EVisibility::Collapsed)
+	{
+		return DockTabStyle->ActiveForegroundColor;
+	}
+	else if (IsHovered())
+	{
+		return DockTabStyle->HoveredForegroundColor;
+	}
+	return FSlateColor::UseStyle();
+}
+
+EVisibility SSidebarButton::GetActiveTabIndicatorVisibility() const
+{
+	const TSharedPtr<FSidebarDrawer> Drawer = DrawerWeak.Pin();
+	if (Drawer.IsValid() && Drawer->bIsOpen)
+	{
+		return EVisibility::HitTestInvisible;
+	}
+	return EVisibility::Collapsed;
+}
+
+EVisibility SSidebarButton::GetPinVisibility() const
+{
+	const TSharedPtr<FSidebarDrawer> Drawer = DrawerWeak.Pin();
+	if (!Drawer.IsValid() || Drawer->bDisablePin)
+	{
+		return EVisibility::Collapsed;
+	}
+	return (Drawer->bIsPinned || IsHovered() || Drawer->bIsOpen)
+		? EVisibility::Visible : EVisibility::Hidden;
+}
+
+FText SSidebarButton::GetPinToolTipText() const
+{
+	const TSharedPtr<FSidebarDrawer> Drawer = DrawerWeak.Pin();
+	if (Drawer.IsValid() && Drawer->bIsPinned)
+	{
+		return LOCTEXT("UnpinTabToolTip", "Unpin Tab");
+	}
+	return LOCTEXT("PinTabToolTip", "Pin Tab");
+}
+
+ECheckBoxState SSidebarButton::IsPinChecked() const
+{
+	const TSharedPtr<FSidebarDrawer> Drawer = DrawerWeak.Pin();
+	if (Drawer.IsValid() && Drawer->bIsPinned)
+	{
+		return ECheckBoxState::Checked;
+	}
+	return ECheckBoxState::Unchecked;
+}
+
+const FSlateBrush* SSidebarButton::GetPinImage() const
+{
+	const TSharedPtr<FSidebarDrawer> Drawer = DrawerWeak.Pin();
+	if (Drawer.IsValid() && Drawer->bIsPinned)
+	{
+		return FAppStyle::Get().GetBrush(TEXT("Icons.Pinned"));
+	}
+	return FAppStyle::Get().GetBrush(TEXT("Icons.Unpinned"));
+}
+
+void SSidebarButton::OnPinStateChanged(const ECheckBoxState InNewState)
+{
+	OnPinToggled.ExecuteIfBound(DrawerWeak.Pin().ToSharedRef(), InNewState == ECheckBoxState::Checked);
+}
+
+EVisibility SSidebarButton::GetDockVisibility() const
+{
+	const TSharedPtr<FSidebarDrawer> Drawer = DrawerWeak.Pin();
+	if (!Drawer.IsValid() || Drawer->bDisableDock)
+	{
+		return EVisibility::Collapsed;
+	}
+	return (Drawer->bIsDocked || IsHovered() || Drawer->bIsOpen) ? EVisibility::Visible : EVisibility::Hidden;
+}
+
+FText SSidebarButton::GetDockToolTipText() const
+{
+	const TSharedPtr<FSidebarDrawer> Drawer = DrawerWeak.Pin();
+	if (Drawer.IsValid() && Drawer->bIsDocked)
+	{
+		return LOCTEXT("UndockTabToolTip", "Undock Tab");
+	}
+	return LOCTEXT("DockTabToolTip", "Dock Tab");
+}
+
+ECheckBoxState SSidebarButton::IsDockChecked() const
+{
+	const TSharedPtr<FSidebarDrawer> Drawer = DrawerWeak.Pin();
+	if (Drawer.IsValid() && Drawer->bIsDocked)
+	{
+		return ECheckBoxState::Checked;
+	}
+	return ECheckBoxState::Unchecked;
+}
+
+const FSlateBrush* SSidebarButton::GetDockImage() const
+{
+	const TSharedPtr<FSidebarDrawer> Drawer = DrawerWeak.Pin();
+	if (Drawer.IsValid() && Drawer->bIsPinned)
+	{
+		return FAppStyle::Get().GetBrush(TEXT("Icons.Layout"));
+	}
+	return FAppStyle::Get().GetBrush(TEXT("Icons.Layout"));
+}
+
+void SSidebarButton::OnDockStateChanged(const ECheckBoxState InNewState)
+{
+	OnDockToggled.ExecuteIfBound(DrawerWeak.Pin().ToSharedRef(), InNewState == ECheckBoxState::Checked);
+}
+
+#undef LOCTEXT_NAMESPACE
diff --git a/Engine/Source/Developer/ToolWidgets/Private/Sidebar/SSidebarButton.h b/Engine/Source/Developer/ToolWidgets/Private/Sidebar/SSidebarButton.h
new file mode 100644
--- /dev/null
+++ b/Engine/Source/Developer/ToolWidgets/Private/Sidebar/SSidebarButton.h
@@ -0,0 +1,86 @@
+// Copyright Epic Games, Inc. All Rights Reserved.
+
+#pragma once
+
+#include "Framework/SlateDelegates.h"
+#include "Misc/Optional.h"
+#include "Sidebar/SidebarDrawer.h"
+#include "Sidebar/SSidebar.h"
+#include "Styling/CoreStyle.h"
+#include "Styling/SlateTypes.h"
+#include "Styling/SlateWidgetStyleAsset.h"
+#include "Widgets/SCompoundWidget.h"
+#include "Widgets/SLeafWidget.h"
+#include "Widgets/DeclarativeSyntaxSupport.h"
+
+class FSlateTextBlockLayout;
+class SBorder;
+class SButton;
+class SCheckBox;
+class SSidebarButtonText;
+class SWidget;
+enum class ESidebarTabLocation : uint8;
+struct FWidgetDrawerConfig;
+
+DECLARE_DELEGATE_OneParam(FOnSidebarButtonPressed, const TSharedRef<FSidebarDrawer>&);
+DECLARE_DELEGATE_TwoParams(FOnSidebarPinToggled, const TSharedRef<FSidebarDrawer>&, const bool /*bInIsPinned*/);
+DECLARE_DELEGATE_TwoParams(FOnSidebarDockToggled, const TSharedRef<FSidebarDrawer>&, const bool /*bInIsDocked*/);
+
+class SSidebarButton : public SCompoundWidget
+{
+public:
+	SLATE_BEGIN_ARGS(SSidebarButton)
+	{}
+		SLATE_ARGUMENT(float, MinButtonSize)
+		SLATE_ARGUMENT(float, MaxButtonSize)
+		SLATE_ARGUMENT(float, ButtonThickness)
+		SLATE_EVENT(FOnSidebarButtonPressed, OnPressed)
+		SLATE_EVENT(FOnSidebarPinToggled, OnPinToggled)
+		SLATE_EVENT(FOnSidebarPinToggled, OnDockToggled)
+		SLATE_EVENT(FOnGetContent, OnGetContextMenuContent)
+	SLATE_END_ARGS()
+
+	void Construct(const FArguments& InArgs, const TSharedRef<FSidebarDrawer>& InDrawer, const ESidebarTabLocation InTabLocation);
+
+	void UpdateAppearance(const TSharedPtr<FSidebarDrawer>& InOpenedDrawer);
+
+	void OnTabRenamed(const TWeakPtr<FSidebarDrawer>& InDrawer);
+	
+protected:
+	EVisibility GetActiveTabIndicatorVisibility() const;
+
+	EVisibility GetPinVisibility() const;
+	FText GetPinToolTipText() const;
+	ECheckBoxState IsPinChecked() const;
+	const FSlateBrush* GetPinImage() const;
+	void OnPinStateChanged(const ECheckBoxState InNewState);
+
+	EVisibility GetDockVisibility() const;
+	FText GetDockToolTipText() const;
+	ECheckBoxState IsDockChecked() const;
+	const FSlateBrush* GetDockImage() const;
+	void OnDockStateChanged(const ECheckBoxState InNewState);
+
+	//~ Begin SWidget
+	virtual FSlateColor GetForegroundColor() const override;
+	virtual FReply OnMouseButtonDown(const FGeometry& InGeometry, const FPointerEvent& InMouseEvent) override;
+	//~ End SWidget
+
+	TWeakPtr<FSidebarDrawer> DrawerWeak;
+	ESidebarTabLocation TabLocation = ESidebarTabLocation::Right;
+
+	FOnSidebarButtonPressed OnPressed;
+	FOnSidebarPinToggled OnPinToggled;
+	FOnSidebarDockToggled OnDockToggled;
+	FOnGetContent OnGetContextMenuContent;
+
+	const FDockTabStyle* DockTabStyle = nullptr;
+
+	TSharedPtr<SSidebarButtonText> Label;
+	TSharedPtr<SWidget> ActiveIndicator;
+	TSharedPtr<SBorder> OpenBorder;
+	TSharedPtr<SButton> MainButton;
+
+	TSharedPtr<SCheckBox> PinCheckBox;
+	TSharedPtr<SCheckBox> DockCheckBox;
+};
diff --git a/Engine/Source/Developer/ToolWidgets/Private/Sidebar/SSidebarButtonText.cpp b/Engine/Source/Developer/ToolWidgets/Private/Sidebar/SSidebarButtonText.cpp
new file mode 100644
--- /dev/null
+++ b/Engine/Source/Developer/ToolWidgets/Private/Sidebar/SSidebarButtonText.cpp
@@ -0,0 +1,88 @@
+// Copyright Epic Games, Inc. All Rights Reserved.
+
+#include "SSidebarButtonText.h"
+#include "Framework/Application/SlateApplication.h"
+#include "Framework/Text/PlainTextLayoutMarshaller.h"
+#include "Widgets/Text/SlateTextBlockLayout.h"
+
+#define LOCTEXT_NAMESPACE "SSidebarButtonText"
+
+void SSidebarButtonText::Construct(const FArguments& InArgs)
+{
+	Text = InArgs._Text;
+	TextStyle = *InArgs._TextStyle;
+	AngleDegrees = InArgs._AngleDegrees;
+
+	TextLayoutCache = MakeUnique<FSlateTextBlockLayout>(this
+		, FTextBlockStyle::GetDefault()
+		, TOptional<ETextShapingMethod>()
+		, TOptional<ETextFlowDirection>()
+		, FCreateSlateTextLayout()
+		, FPlainTextLayoutMarshaller::Create()
+		, nullptr);
+	
+	TextLayoutCache->SetTextOverflowPolicy(InArgs._OverflowPolicy.IsSet() ? InArgs._OverflowPolicy : TextStyle.OverflowPolicy);
+}
+
+int32 SSidebarButtonText::OnPaint(const FPaintArgs& InArgs, const FGeometry& InAllottedGeometry, const FSlateRect& InCullingRect
+	, FSlateWindowElementList& OutDrawElements, int32 InLayerId, const FWidgetStyle& InWidgetStyle, bool bInParentEnabled) const
+{
+	// We're going to figure out the bounds of the corresponding horizontal text, and then rotate it into a vertical orientation.
+	const FVector2D LocalSize = InAllottedGeometry.GetLocalSize();
+	const FVector2D DesiredHorizontalTextSize = TextLayoutCache->GetDesiredSize();
+	const FVector2D ActualHorizontalTextSize(FMath::Min(DesiredHorizontalTextSize.X, LocalSize.Y), FMath::Min(DesiredHorizontalTextSize.Y, LocalSize.X));
+
+	// Now determine the center of the vertical text by rotating the dimensions of the horizontal text.
+	// The center should align it to the top of the widget.
+	const FVector2D VerticalTextSize(ActualHorizontalTextSize.Y, ActualHorizontalTextSize.X);
+	const FVector2D VerticalTextCenter = VerticalTextSize * 0.5f;
+
+	// Now determine where the horizontal text should be positioned so that it is centered on the vertical text:
+	//      +-+
+	//      |v|
+	//      |e|
+	// [ horizontal ]
+	//      |r|
+	//      |t|
+	//      +-+
+	const FVector2D HorizontalTextPosition = VerticalTextCenter - (ActualHorizontalTextSize * 0.5f);
+
+	// Define the text's geometry using the horizontal bounds, then rotate it 90/-90 degrees into place to become vertical.
+	const FSlateRenderTransform RotationTransform(FSlateRenderTransform(FQuat2D(FMath::DegreesToRadians(AngleDegrees.Get(0.f)))));
+	const FGeometry TextGeometry = InAllottedGeometry.MakeChild(ActualHorizontalTextSize
+		, FSlateLayoutTransform(HorizontalTextPosition), RotationTransform, FVector2D(0.5f, 0.5f));
+
+	return TextLayoutCache->OnPaint(InArgs, TextGeometry, InCullingRect, OutDrawElements, InLayerId, InWidgetStyle, ShouldBeEnabled(bInParentEnabled));
+}
+
+FVector2D SSidebarButtonText::ComputeDesiredSize(const float LayoutScaleMultiplier) const
+{
+	// The text's desired size reflects the horizontal/untransformed text.
+	// Switch the dimensions for vertical text.
+	const FVector2D DesiredHorizontalTextSize = TextLayoutCache->ComputeDesiredSize(
+		FSlateTextBlockLayout::FWidgetDesiredSizeArgs(
+			Text.Get(),
+			FText(),
+			0.f,
+			false,
+			ETextWrappingPolicy::DefaultWrapping,
+			ETextTransformPolicy::None,
+			FMargin(),
+			1.f,
+			true,
+			ETextJustify::Left),
+		LayoutScaleMultiplier, TextStyle);
+	return FVector2D(DesiredHorizontalTextSize.Y, DesiredHorizontalTextSize.X);
+}
+
+void SSidebarButtonText::SetText(const TAttribute<FText>& InText)
+{
+	Text = InText;
+}
+
+void SSidebarButtonText::SetRotation(const TAttribute<float>& InAngleDegrees)
+{
+	AngleDegrees = InAngleDegrees;
+}
+
+#undef LOCTEXT_NAMESPACE
diff --git a/Engine/Source/Developer/ToolWidgets/Private/Sidebar/SSidebarButtonText.h b/Engine/Source/Developer/ToolWidgets/Private/Sidebar/SSidebarButtonText.h
new file mode 100644
--- /dev/null
+++ b/Engine/Source/Developer/ToolWidgets/Private/Sidebar/SSidebarButtonText.h
@@ -0,0 +1,56 @@
+﻿// Copyright Epic Games, Inc. All Rights Reserved.
+
+#pragma once
+
+#include "Framework/SlateDelegates.h"
+#include "Misc/Attribute.h"
+#include "Misc/Optional.h"
+#include "Styling/CoreStyle.h"
+#include "Styling/SlateTypes.h"
+#include "Styling/SlateWidgetStyleAsset.h"
+#include "Templates/UniquePtr.h"
+#include "Widgets/DeclarativeSyntaxSupport.h"
+#include "Widgets/SLeafWidget.h"
+#include "Widgets/Text/SlateTextBlockLayout.h"
+
+class FText;
+
+/**
+ * Vertical text block for use in the tab drawer button.
+ * Text is aligned to the top of the widget if it fits without clipping;
+ * otherwise it is ellipsized and fills the widget height.
+ */
+class SSidebarButtonText : public SLeafWidget
+{
+public:
+	SLATE_BEGIN_ARGS(SSidebarButtonText)
+		: _Text()
+		, _TextStyle(&FCoreStyle::Get().GetWidgetStyle<FTextBlockStyle>(TEXT("NormalText")))
+		, _AngleDegrees(0)
+		, _OverflowPolicy()
+	{}
+		SLATE_ATTRIBUTE(FText, Text)
+		SLATE_STYLE_ARGUMENT(FTextBlockStyle, TextStyle)
+		SLATE_ATTRIBUTE(float, AngleDegrees)
+		SLATE_ARGUMENT(TOptional<ETextOverflowPolicy>, OverflowPolicy)
+	SLATE_END_ARGS()
+
+	void Construct(const FArguments& InArgs);
+
+	void SetText(const TAttribute<FText>& InText);
+
+	void SetRotation(const TAttribute<float>& InAngleDegrees);
+
+protected:
+	//~ Begin SWidget
+	virtual int32 OnPaint(const FPaintArgs& InArgs, const FGeometry& InAllottedGeometry, const FSlateRect& InCullingRect
+		, FSlateWindowElementList& OutDrawElements, int32 InLayerId, const FWidgetStyle& InWidgetStyle, bool bInParentEnabled) const override;
+	virtual FVector2D ComputeDesiredSize(const float InLayoutScaleMultiplier) const override;
+	//~ End SWidget
+
+	TAttribute<FText> Text;
+	FTextBlockStyle TextStyle;
+	TAttribute<float> AngleDegrees;
+	
+	TUniquePtr<FSlateTextBlockLayout> TextLayoutCache;
+};
diff --git a/Engine/Source/Developer/ToolWidgets/Private/Sidebar/SSidebarDrawer.cpp b/Engine/Source/Developer/ToolWidgets/Private/Sidebar/SSidebarDrawer.cpp
new file mode 100644
--- /dev/null
+++ b/Engine/Source/Developer/ToolWidgets/Private/Sidebar/SSidebarDrawer.cpp
@@ -0,0 +1,605 @@
+// Copyright Epic Games, Inc. All Rights Reserved.
+
+#include "SSidebarDrawer.h"
+#include "Framework/Application/SlateApplication.h"
+#include "Sidebar/SidebarDrawer.h"
+#include "Sidebar/SSidebar.h"
+#include "Styling/AppStyle.h"
+#include "Widgets/Layout/SBox.h"
+
+SSidebarDrawer::~SSidebarDrawer()
+{
+	FSlateApplication::Get().OnFocusChanging().RemoveAll(this);
+
+	FSlateThrottleManager::Get().LeaveResponsiveMode(AnimationThrottle);
+	FSlateThrottleManager::Get().LeaveResponsiveMode(ResizeThrottleHandle);
+}
+
+void SSidebarDrawer::Construct(const FArguments& InArgs, const TSharedRef<FSidebarDrawer>& InDrawer, const ESidebarTabLocation InTabLocation)
+{
+	check(InDrawer->ContentWidget.IsValid());
+
+	DrawerWeak = InDrawer;
+	TabLocation = InTabLocation;
+
+	MinDrawerSize = InArgs._MinDrawerSize;
+	MaxDrawerSize = InArgs._MaxDrawerSize;
+	TargetDrawerSize = FMath::Clamp(InArgs._TargetDrawerSize, MinDrawerSize, MaxDrawerSize);
+	ShadowOffset = InArgs._ShadowOffset;
+	ExpanderHandleSize = InArgs._ExpanderHandleSize;
+	
+	OnDrawerTargetSizeChanged = InArgs._OnDrawerTargetSizeChanged;
+	OnDrawerFocused = InArgs._OnDrawerFocused;
+	OnDrawerFocusLost = InArgs._OnDrawerFocusLost;
+	OnDrawerClosed = InArgs._OnDrawerClosed;
+
+	OpenCloseAnimation = FCurveSequence(0.f, 0.15f, ECurveEaseFunction::QuadOut);
+
+	FSlateApplication::Get().OnFocusChanging().AddSP(this, &SSidebarDrawer::OnGlobalFocusChanging);
+
+	ChildSlot
+	[
+		SNew(SBox)
+		.Clipping(EWidgetClipping::ClipToBounds)
+		[
+			InDrawer->ContentWidget.ToSharedRef()
+		]
+	];
+}
+
+void SSidebarDrawer::SetCurrentSize(const float InSize)
+{
+	CurrentSize = FMath::Clamp(InSize, MinDrawerSize, TargetDrawerSize);
+}
+
+void SSidebarDrawer::Open(bool bAnimateOpen)
+{
+	if (!bAnimateOpen)
+	{
+		SetCurrentSize(TargetDrawerSize);
+		OpenCloseAnimation.JumpToEnd();
+		return;
+	}
+
+	OpenCloseAnimation.Play(AsShared(), false, OpenCloseAnimation.IsPlaying() ? OpenCloseAnimation.GetSequenceTime() : 0.0f, false);
+
+	if (!OpenCloseTimer.IsValid())
+	{
+		AnimationThrottle = FSlateThrottleManager::Get().EnterResponsiveMode();
+		OpenCloseTimer = RegisterActiveTimer(0.0f, FWidgetActiveTimerDelegate::CreateSP(this, &SSidebarDrawer::UpdateAnimation));
+	}
+}
+
+void SSidebarDrawer::Close()
+{
+	if (OpenCloseAnimation.IsForward())
+	{
+		OpenCloseAnimation.Reverse();
+	}
+
+	if (!OpenCloseTimer.IsValid())
+	{
+		AnimationThrottle = FSlateThrottleManager::Get().EnterResponsiveMode();
+		OpenCloseTimer = RegisterActiveTimer(0.0f, FWidgetActiveTimerDelegate::CreateSP(this, &SSidebarDrawer::UpdateAnimation));
+	}
+}
+
+bool SSidebarDrawer::IsOpen() const
+{
+	return !OpenCloseAnimation.IsAtStart();
+}
+
+bool SSidebarDrawer::IsClosing() const
+{
+	return OpenCloseAnimation.IsPlaying() && OpenCloseAnimation.IsInReverse();
+}
+
+TSharedPtr<FSidebarDrawer> SSidebarDrawer::GetDrawer() const
+{
+	return DrawerWeak.Pin();
+}
+
+bool SSidebarDrawer::SupportsKeyboardFocus() const 
+{
+	return true;
+}
+
+FVector2D SSidebarDrawer::ComputeDesiredSize(const float InLayoutScaleMultiplier) const
+{
+	switch (TabLocation)
+	{
+	case ESidebarTabLocation::Left:
+	case ESidebarTabLocation::Right:
+		return FVector2D(TargetDrawerSize + ShadowOffset.X, 1.f);
+	case ESidebarTabLocation::Top:
+	case ESidebarTabLocation::Bottom:
+		return FVector2D(1.f, TargetDrawerSize + ShadowOffset.Y);
+	}
+	return FVector2D::One();
+}
+
+void SSidebarDrawer::OnArrangeChildren(const FGeometry& InAllottedGeometry, FArrangedChildren& ArrangedChildren) const 
+{
+	const EVisibility ChildVisibility = ChildSlot.GetWidget()->GetVisibility();
+	if (!ArrangedChildren.Accepts(ChildVisibility))
+	{
+		return;
+	}
+
+	FVector2D ChildOffset;
+	FVector2D LocalSize;
+
+	switch (TabLocation)
+	{
+	case ESidebarTabLocation::Left:
+		ChildOffset = FVector2D(0.f, ShadowOffset.Y);
+		LocalSize = FVector2D(TargetDrawerSize, InAllottedGeometry.GetLocalSize().Y - (ShadowOffset.Y * 2));
+		break;
+	case ESidebarTabLocation::Right:
+		ChildOffset = ShadowOffset;
+		LocalSize = FVector2D(TargetDrawerSize, InAllottedGeometry.GetLocalSize().Y - (ShadowOffset.Y * 2));
+		break;
+	case ESidebarTabLocation::Top:
+		ChildOffset = FVector2D(ShadowOffset.X, 0.f);
+		LocalSize = FVector2D(InAllottedGeometry.GetLocalSize().X - (ShadowOffset.X * 2), TargetDrawerSize);
+		break;
+	case ESidebarTabLocation::Bottom:
+		ChildOffset = ShadowOffset;
+		LocalSize = FVector2D(InAllottedGeometry.GetLocalSize().X - (ShadowOffset.X * 2), TargetDrawerSize);
+		break;
+	}
+
+	ArrangedChildren.AddWidget(InAllottedGeometry.MakeChild(ChildSlot.GetWidget(), ChildOffset, LocalSize));
+}
+
+FReply SSidebarDrawer::OnMouseButtonDown(const FGeometry& InAllottedGeometry, const FPointerEvent& InMouseEvent) 
+{
+	FReply Reply = FReply::Unhandled();
+
+	if (InMouseEvent.GetEffectingButton() == EKeys::LeftMouseButton)
+	{
+		const FGeometry ResizeHandleGeometry = GetResizeHandleGeometry(InAllottedGeometry);
+
+		if (ResizeHandleGeometry.IsUnderLocation(InMouseEvent.GetScreenSpacePosition()))
+		{
+			bIsResizing = true;
+			InitialResizeGeometry = ResizeHandleGeometry;
+			InitialSizeAtResize = CurrentSize;
+			ResizeThrottleHandle = FSlateThrottleManager::Get().EnterResponsiveMode();
+
+			Reply = FReply::Handled().CaptureMouse(SharedThis(this));
+		}
+	}
+
+	return Reply;
+}
+
+FReply SSidebarDrawer::OnMouseButtonUp(const FGeometry& InAllottedGeometry, const FPointerEvent& InMouseEvent) 
+{
+	if (bIsResizing && InMouseEvent.GetEffectingButton() == EKeys::LeftMouseButton)
+	{
+		bIsResizing = false;
+		FSlateThrottleManager::Get().LeaveResponsiveMode(ResizeThrottleHandle);
+
+		OnDrawerTargetSizeChanged.ExecuteIfBound(SharedThis(this), TargetDrawerSize);
+		return FReply::Handled().ReleaseMouseCapture();
+	}
+
+	return FReply::Unhandled();
+}
+
+FReply SSidebarDrawer::OnMouseMove(const FGeometry& InAllottedGeometry, const FPointerEvent& InMouseEvent)
+{
+	const FGeometry ResizeHandleGeometry = GetResizeHandleGeometry(InAllottedGeometry);
+
+	bIsResizeHandleHovered = ResizeHandleGeometry.IsUnderLocation(InMouseEvent.GetScreenSpacePosition());
+
+	if (bIsResizing && this->HasMouseCapture() && !InMouseEvent.GetCursorDelta().IsZero())
+	{
+		const FVector2D MousePosition = InMouseEvent.GetScreenSpacePosition();
+		const FVector2D LocalMousePosition = InitialResizeGeometry.AbsoluteToLocal(MousePosition);
+
+		float DeltaSize = 0.0f;
+
+		switch (TabLocation)
+		{
+		case ESidebarTabLocation::Left:
+			DeltaSize = LocalMousePosition.X;
+			break;
+		case ESidebarTabLocation::Right:
+			DeltaSize = -LocalMousePosition.X;
+			break;
+		case ESidebarTabLocation::Top:
+			DeltaSize = LocalMousePosition.Y;
+			break;
+		case ESidebarTabLocation::Bottom:
+			DeltaSize = -LocalMousePosition.Y;
+			break;
+		}
+
+		TargetDrawerSize = FMath::Clamp(InitialSizeAtResize + DeltaSize, MinDrawerSize, MaxDrawerSize);
+		SetCurrentSize(InitialSizeAtResize + DeltaSize);
+
+		return FReply::Handled();
+	}
+
+	return FReply::Unhandled();
+}
+
+void SSidebarDrawer::OnMouseLeave(const FPointerEvent& InMouseEvent)
+{
+	SCompoundWidget::OnMouseLeave(InMouseEvent);
+
+	bIsResizeHandleHovered = false;
+}
+
+FCursorReply SSidebarDrawer::OnCursorQuery(const FGeometry& InGeometry, const FPointerEvent& InCursorEvent) const
+{
+	if (bIsResizing || bIsResizeHandleHovered)
+	{
+		switch (TabLocation)
+		{
+		case ESidebarTabLocation::Left:
+		case ESidebarTabLocation::Right:
+			return FCursorReply::Cursor(EMouseCursor::ResizeLeftRight);
+		case ESidebarTabLocation::Top:
+		case ESidebarTabLocation::Bottom:
+			return FCursorReply::Cursor(EMouseCursor::ResizeUpDown);
+		}
+	}
+
+	return FCursorReply::Unhandled();
+}
+
+int32 SSidebarDrawer::OnPaint(const FPaintArgs& InArgs, const FGeometry& InAllottedGeometry, const FSlateRect& InCullingRect
+	, FSlateWindowElementList& OutDrawElements, int32 InLayerId, const FWidgetStyle& InWidgetStyle, bool bInParentEnabled) const
+{
+	static const FSlateBrush* BackgroundBrush = FAppStyle::Get().GetBrush(TEXT("Docking.Sidebar.DrawerBackground"));
+	static const FSlateBrush* ShadowBrush = FAppStyle::Get().GetBrush(TEXT("Docking.Sidebar.DrawerShadow"));
+	static const FSlateBrush* BorderBrush = FAppStyle::Get().GetBrush(TEXT("Docking.Sidebar.Border"));
+	static const FSlateBrush* BorderSquareEdgeBrush;
+	switch (TabLocation)
+	{
+	case ESidebarTabLocation::Left:
+		BorderSquareEdgeBrush = FAppStyle::Get().GetBrush(TEXT("Docking.Sidebar.Border_SquareLeft"));
+		break;
+	case ESidebarTabLocation::Right:
+		BorderSquareEdgeBrush = FAppStyle::Get().GetBrush(TEXT("Docking.Sidebar.Border_SquareRight"));
+		break;
+	case ESidebarTabLocation::Top:
+	case ESidebarTabLocation::Bottom:
+		// @TODO: There are no existing top or bottom brushes
+		BorderSquareEdgeBrush = BorderBrush;
+		break;
+	}
+	static const FSplitterStyle* SplitterStyle = &FAppStyle::Get().GetWidgetStyle<FSplitterStyle>(TEXT("Splitter"));
+	static const FSlateColor ShadowColor = FAppStyle::Get().GetSlateColor(TEXT("Colors.Foldout"));
+
+	const FGeometry RenderTransformedChildGeometry = GetRenderTransformedGeometry(InAllottedGeometry);
+	const FGeometry ResizeHandleGeometry = GetResizeHandleGeometry(InAllottedGeometry);
+	const FVector2D LocalSize = InAllottedGeometry.GetLocalSize();
+
+	FVector2D ContentsLocalOrigin;
+	FVector2D ContentsLocalSize;
+
+	switch (TabLocation)
+	{
+	case ESidebarTabLocation::Left:
+		ContentsLocalOrigin = FVector2D(0.f, ShadowOffset.Y);
+		ContentsLocalSize = FVector2D(TargetDrawerSize, LocalSize.Y - (ShadowOffset.Y * 2));
+		break;
+	case ESidebarTabLocation::Right:
+		ContentsLocalOrigin = ShadowOffset;
+		ContentsLocalSize = FVector2D(TargetDrawerSize, LocalSize.Y - (ShadowOffset.Y * 2));
+		break;
+	case ESidebarTabLocation::Top:
+	case ESidebarTabLocation::Bottom:
+		ContentsLocalOrigin = ShadowOffset;
+		ContentsLocalSize = FVector2D(LocalSize.X - (ShadowOffset.X * 2), TargetDrawerSize);
+		break;
+	}
+
+	const FPaintGeometry OffsetPaintGeom = RenderTransformedChildGeometry.ToPaintGeometry(ContentsLocalSize, FSlateLayoutTransform(ContentsLocalOrigin));
+
+	// Draw the resize handle
+	if (bIsResizing || bIsResizeHandleHovered)
+	{
+		const FSlateBrush* SplitterBrush = &SplitterStyle->HandleHighlightBrush;
+		FSlateDrawElement::MakeBox(
+			OutDrawElements,
+			InLayerId,
+			ResizeHandleGeometry.ToPaintGeometry(),
+			SplitterBrush,
+			ESlateDrawEffect::None,
+			SplitterBrush->GetTint(InWidgetStyle));
+	}
+
+	// Main Shadow
+	FSlateDrawElement::MakeBox(
+		OutDrawElements,
+		InLayerId,
+		RenderTransformedChildGeometry.ToPaintGeometry(),
+		ShadowBrush,
+		ESlateDrawEffect::None,
+		ShadowBrush->GetTint(InWidgetStyle));
+
+	// Background
+	FSlateDrawElement::MakeBox(
+		OutDrawElements,
+		InLayerId,
+		OffsetPaintGeom,
+		BackgroundBrush,
+		ESlateDrawEffect::None,
+		BackgroundBrush->GetTint(InWidgetStyle));
+
+	int32 OutLayerId = SCompoundWidget::OnPaint(InArgs, RenderTransformedChildGeometry, InCullingRect, OutDrawElements, InLayerId, InWidgetStyle, bInParentEnabled);
+
+	TSharedPtr<SWidget> TabButton;
+	if (const TSharedPtr<FSidebarDrawer> Drawer = DrawerWeak.Pin())
+	{
+		TabButton = Drawer->ButtonWidget;
+	}
+	if (!TabButton.IsValid())
+	{
+		return ++OutLayerId;
+	}
+
+	// Top border
+	switch (TabLocation)
+	{
+	case ESidebarTabLocation::Left:
+	case ESidebarTabLocation::Right:
+		{
+			// Example of how border box is drawn with the tab notch cut out on the right side
+			// (OpenDirection == ESidebarTabLocation::Right)
+			//
+			//                    + - - - - - - +
+			//                    : /---------\ :
+			// ClipAboveTabButton : |         | : 
+			//                    : |         | :
+			//            TabTopY + - - - - - - +
+			//                    : |           :  |
+			// ClipAtTabButton    : |           :  |  (right edge outside clip is clipped off)
+			//                    : |           :  |
+			//         TabBottomY + - - - - - - +
+			//                    : |         | :
+			// ClipBelowTabButton : |         | :
+			//                    : \---------/ :
+			//                    + - - - - - - +
+			//                                  <-->
+			//                              NotchOffset
+			//
+			// Originally, I tried making the middle clip region thinner (to clip out the notch)
+			// while keeping the geometry identical, but this looks worse when the tab notch needs to
+			// be at the top or bottom, since the top/bottom edge of the border wouldn't extend all the
+			// way to the edge.
+			const FGeometry TabButtonGeometry = TabButton->GetPaintSpaceGeometry();
+
+			// Compute the top/bottom of the tab in our local space.
+			const float BorderWidth = BorderBrush->OutlineSettings.Width;
+			const float TabTopY = RenderTransformedChildGeometry.AbsoluteToLocal(TabButtonGeometry.GetAbsolutePositionAtCoordinates(FVector2D::ZeroVector)).Y + 0.5f * BorderWidth;
+			const float TabBottomY = RenderTransformedChildGeometry.AbsoluteToLocal(TabButtonGeometry.GetAbsolutePositionAtCoordinates(FVector2D::UnitVector)).Y - 0.5f * BorderWidth;
+
+			// Create the geometry for the notched portion, where one edge extends past the clipping rect.
+			const FVector2D NotchOffsetSize(TabButtonGeometry.GetLocalSize().X, 0.f);
+			const FVector2D NotchOffsetTranslate = TabLocation == ESidebarTabLocation::Left ? -NotchOffsetSize : FVector2D::ZeroVector;
+			const FPaintGeometry NotchOffsetPaintGeom = RenderTransformedChildGeometry.ToPaintGeometry(ContentsLocalSize + NotchOffsetSize
+				, FSlateLayoutTransform(ContentsLocalOrigin + NotchOffsetTranslate));
+
+			// Split the border box into three clipping zones.
+			const FPaintGeometry ClipAboveTabButton = RenderTransformedChildGeometry.ToPaintGeometry(FVector2f(LocalSize.X, TabTopY)
+				, FSlateLayoutTransform(FVector2f(0.f, 0.f)));
+			const FPaintGeometry ClipAtTabButton = RenderTransformedChildGeometry.ToPaintGeometry(FVector2f(LocalSize.X, TabBottomY - TabTopY)
+				, FSlateLayoutTransform(FVector2f(0.f, TabTopY)));
+			const FPaintGeometry ClipBelowTabButton = RenderTransformedChildGeometry.ToPaintGeometry(FVector2f(LocalSize.X, LocalSize.Y - TabBottomY)
+				, FSlateLayoutTransform(FVector2f(0.f, TabBottomY)));
+
+			// If the tab button touches a corner on the edge of the border, switch the brush to
+			// draw that corner squared-off. When a tab is near the very top or bottom of its sidebar,
+			// this makes the outline look slightly nicer and more connected.
+			const int32 UpperCornerIndex = TabLocation == ESidebarTabLocation::Left ? 0 : 1;
+			const int32 LowerCornerIndex = TabLocation == ESidebarTabLocation::Left ? 3 : 2;
+			const bool bTabTouchesUpperCorner = TabTopY < ShadowOffset.Y + BorderBrush->OutlineSettings.CornerRadii[UpperCornerIndex];
+			const bool bTabTouchesLowerCorner = TabBottomY > LocalSize.Y - ShadowOffset.Y - BorderBrush->OutlineSettings.CornerRadii[LowerCornerIndex];
+			const FSlateBrush* AboveTabBrush = bTabTouchesUpperCorner ? BorderSquareEdgeBrush : BorderBrush;
+			const FSlateBrush* BelowTabBrush = bTabTouchesLowerCorner ? BorderSquareEdgeBrush : BorderBrush;
+
+			// Draw portion above the tab
+			OutDrawElements.PushClip(FSlateClippingZone(ClipAboveTabButton));
+			FSlateDrawElement::MakeBox(OutDrawElements, OutLayerId, OffsetPaintGeom, AboveTabBrush
+				, ESlateDrawEffect::None, AboveTabBrush->GetTint(InWidgetStyle));
+			OutDrawElements.PopClip();
+
+			// Draw "notched" portion next to the tab
+			OutDrawElements.PushClip(FSlateClippingZone(ClipAtTabButton));
+			FSlateDrawElement::MakeBox(OutDrawElements, OutLayerId, NotchOffsetPaintGeom, BorderSquareEdgeBrush
+				, ESlateDrawEffect::None, BorderSquareEdgeBrush->GetTint(InWidgetStyle));
+			OutDrawElements.PopClip();
+
+			// Draw portion below the tab
+			OutDrawElements.PushClip(FSlateClippingZone(ClipBelowTabButton));
+			FSlateDrawElement::MakeBox(OutDrawElements, OutLayerId, OffsetPaintGeom, BelowTabBrush
+				, ESlateDrawEffect::None, BelowTabBrush->GetTint(InWidgetStyle));
+			OutDrawElements.PopClip();
+		}
+		break;
+	case ESidebarTabLocation::Top:
+	case ESidebarTabLocation::Bottom:
+		{
+			// When opened from the bottom, draw the full border.
+			// Cutting out the "notch" for the corresponding tab is only supported in left/right orientations.
+			FSlateDrawElement::MakeBox(OutDrawElements, OutLayerId, OffsetPaintGeom, BorderBrush
+				, ESlateDrawEffect::None, BorderBrush->GetTint(InWidgetStyle));
+		}
+		break;
+	}
+
+	return ++OutLayerId;
+}
+
+FGeometry SSidebarDrawer::GetRenderTransformedGeometry(const FGeometry& InAllottedGeometry) const
+{
+	switch (TabLocation)
+	{
+	case ESidebarTabLocation::Left:
+		return InAllottedGeometry.MakeChild(FSlateRenderTransform(FVector2D(CurrentSize - TargetDrawerSize, 0.f)));
+	case ESidebarTabLocation::Right:
+		return InAllottedGeometry.MakeChild(FSlateRenderTransform(FVector2D(TargetDrawerSize - CurrentSize, 0.f)));
+	case ESidebarTabLocation::Top:
+		return InAllottedGeometry.MakeChild(FSlateRenderTransform(FVector2D(0.f, CurrentSize - TargetDrawerSize)));
+	case ESidebarTabLocation::Bottom:
+		return InAllottedGeometry.MakeChild(FSlateRenderTransform(FVector2D(0.f, TargetDrawerSize - CurrentSize)));
+	}
+	return FGeometry();
+}
+
+FGeometry SSidebarDrawer::GetResizeHandleGeometry(const FGeometry& InAllottedGeometry) const
+{
+	const FGeometry RenderTransformedGeometry = GetRenderTransformedGeometry(InAllottedGeometry);
+
+	FVector2D LocalSize;
+	FVector2D Translation;
+
+	switch (TabLocation)
+	{
+	case ESidebarTabLocation::Left:
+		LocalSize = FVector2D(ExpanderHandleSize, InAllottedGeometry.GetLocalSize().Y - (ShadowOffset.Y * 2));
+		Translation = FVector2D(RenderTransformedGeometry.GetLocalSize().X - ShadowOffset.X, ShadowOffset.Y);
+		break;
+	case ESidebarTabLocation::Right:
+		LocalSize = FVector2D(ExpanderHandleSize, InAllottedGeometry.GetLocalSize().Y - (ShadowOffset.Y * 2));
+		Translation = ShadowOffset - FVector2D(ExpanderHandleSize, 0.f);
+		break;
+	case ESidebarTabLocation::Top:
+		LocalSize = FVector2D(InAllottedGeometry.GetLocalSize().X - (ShadowOffset.X * 2), ExpanderHandleSize);
+		Translation = FVector2D(ShadowOffset.X, RenderTransformedGeometry.GetLocalSize().Y - ShadowOffset.Y);
+		break;
+	case ESidebarTabLocation::Bottom:
+		LocalSize = FVector2D(InAllottedGeometry.GetLocalSize().X - (ShadowOffset.X * 2), ExpanderHandleSize);
+		Translation = ShadowOffset - FVector2D(0.f, ExpanderHandleSize);
+		break;
+	}
+
+	return RenderTransformedGeometry.MakeChild(LocalSize, FSlateLayoutTransform(Translation));
+}
+
+EActiveTimerReturnType SSidebarDrawer::UpdateAnimation(const double InCurrentTime, const float InDeltaTime)
+{
+	SetCurrentSize(FMath::Lerp(0.f, TargetDrawerSize, OpenCloseAnimation.GetLerp()));
+
+	if (!OpenCloseAnimation.IsPlaying())
+	{
+		if (OpenCloseAnimation.IsAtStart())
+		{
+			OnDrawerClosed.ExecuteIfBound(SharedThis(this));
+		}
+
+		FSlateThrottleManager::Get().LeaveResponsiveMode(AnimationThrottle);
+		OpenCloseTimer.Reset();
+
+		return EActiveTimerReturnType::Stop;
+	}
+
+	return EActiveTimerReturnType::Continue;
+}
+
+static bool IsLegalWidgetFocused(const FWidgetPath& InFocusPath, const TArrayView<TSharedRef<SWidget>>& InLegalFocusWidgets)
+{
+	for (const TSharedRef<SWidget>& Widget : InLegalFocusWidgets)
+	{
+		if (InFocusPath.ContainsWidget(&Widget.Get()))
+		{
+			return true;
+		}
+	}
+
+	return false;
+}
+
+void SSidebarDrawer::OnGlobalFocusChanging(const FFocusEvent& InFocusEvent
+	, const FWeakWidgetPath& InOldFocusedWidgetPath, const TSharedPtr<SWidget>& InOldFocusedWidget
+	, const FWidgetPath& InNewFocusedWidgetPath, const TSharedPtr<SWidget>& InNewFocusedWidget)
+{
+	// Sometimes when dismissing focus can change which will trigger this again
+	static bool bIsReEntrant = false;
+
+	if (bIsReEntrant)
+	{
+		return;
+	}
+
+	// Do not close due to slow tasks as those opening send window activation events
+	if (GIsSlowTask || FSlateApplication::Get().GetActiveModalWindow().IsValid())
+	{
+		return;
+	}
+
+	TGuardValue<bool> ReEntrancyGuard(bIsReEntrant, true);
+
+	const TSharedRef<SSidebarDrawer> ThisWidget = SharedThis(this);
+
+	TArray<TSharedRef<SWidget>, TInlineAllocator<4>> LegalFocusWidgets;
+	LegalFocusWidgets.Add(ThisWidget);
+	LegalFocusWidgets.Add(ChildSlot.GetWidget());
+	if (const TSharedPtr<FSidebarDrawer> Drawer = DrawerWeak.Pin())
+	{
+		LegalFocusWidgets.Add(Drawer->ButtonWidget.ToSharedRef());
+	}
+
+	bool bShouldLoseFocus = false;
+
+	if (IsLegalWidgetFocused(InNewFocusedWidgetPath, MakeArrayView(LegalFocusWidgets)))
+	{
+		// New focus is on this tab, so make it active
+		if (!IsClosing())
+		{
+			OnDrawerFocused.ExecuteIfBound(ThisWidget);
+		}
+	}
+	else if (InNewFocusedWidgetPath.IsValid())
+	{
+		// New focus is on something else, try to check if it's a menu or child window
+		const TSharedRef<SWindow> NewWindow = InNewFocusedWidgetPath.GetWindow();
+		const TSharedPtr<SWindow> MyWindow = FSlateApplication::Get().FindWidgetWindow(ThisWidget);
+
+		// See if this is a child window (like a color picker being opened from details), and if so, don't dismiss
+		// Rely on OnActiveTabChanged below to lose focus if the child window actually contains tabs
+		if (!NewWindow->IsDescendantOf(MyWindow))
+		{
+			if (const TSharedPtr<SWidget> MenuHost = FSlateApplication::Get().GetMenuHostWidget())
+			{
+				FWidgetPath MenuHostPath;
+
+				// See if the menu being opened is owned by the drawer contents and if so the menu should not be dismissed
+				FSlateApplication::Get().GeneratePathToWidgetUnchecked(MenuHost.ToSharedRef(), MenuHostPath);
+				
+				if (!MenuHostPath.ContainsWidget(&ChildSlot.GetWidget().Get()))
+				{
+					bShouldLoseFocus = true;
+				}
+			}
+			else
+			{
+				bShouldLoseFocus = true;
+			}
+		}
+	}
+	else
+	{
+		bShouldLoseFocus = true;
+	}
+
+	if (bShouldLoseFocus)
+	{
+		OnDrawerFocusLost.ExecuteIfBound(ThisWidget);
+	}
+}
+
+void SSidebarDrawer::OnActiveTabChanged(const TSharedPtr<FSidebarDrawer>& InPreviouslyActive, const TSharedPtr<FSidebarDrawer>& InNewlyActivated)
+{
+	// This tab lost the active status to some other tab; treat this like focus was lost
+	if (InPreviouslyActive == DrawerWeak && InNewlyActivated.IsValid())
+	{
+		OnDrawerFocusLost.ExecuteIfBound(SharedThis(this));
+	}
+}
diff --git a/Engine/Source/Developer/ToolWidgets/Private/Sidebar/SSidebarDrawer.h b/Engine/Source/Developer/ToolWidgets/Private/Sidebar/SSidebarDrawer.h
new file mode 100644
--- /dev/null
+++ b/Engine/Source/Developer/ToolWidgets/Private/Sidebar/SSidebarDrawer.h
@@ -0,0 +1,129 @@
+// Copyright Epic Games, Inc. All Rights Reserved.
+
+#pragma once
+
+#include "SSidebarButton.h"
+#include "Animation/CurveSequence.h"
+#include "Application/ThrottleManager.h"
+#include "Framework/SlateDelegates.h"
+#include "Layout/Geometry.h"
+#include "Templates/SharedPointer.h"
+#include "Widgets/SCompoundWidget.h"
+
+class FSidebarDrawer;
+class ISidebarDrawerContent;
+class SSidebarDrawer;
+enum class ESidebarTabLocation : uint8;
+struct FSplitterStyle;
+
+DECLARE_DELEGATE_OneParam(FGenericSidebarDrawerWidgetDelegate, const TSharedRef<SSidebarDrawer>&);
+DECLARE_DELEGATE_TwoParams(FOnSidebarDrawerTargetSizeChanged, const TSharedRef<SSidebarDrawer>&, float);
+
+/** Handles sliding drawer animation */
+class SSidebarDrawer : public SCompoundWidget
+{
+public:
+	SLATE_BEGIN_ARGS(SSidebarDrawer)
+		: _ShadowOffset(FVector2D(8.f, 8.f))
+		, _ExpanderHandleSize(5.f)
+	{
+		_Clipping = EWidgetClipping::ClipToBounds;
+	}
+		/** The minimum size a drawer can be when opened. This unit is in window space */
+		SLATE_ARGUMENT(float, MinDrawerSize)
+		/** The maximum size a drawer can be when opened. This unit is in window space */
+		SLATE_ARGUMENT(float, MaxDrawerSize)
+		/** The size that the drawer should be, clamped to the above min/max values. This unit is in window space */
+		SLATE_ARGUMENT(float, TargetDrawerSize)
+		/** The side of the drop shadow surrounding the drawer */
+		SLATE_ARGUMENT(FVector2D, ShadowOffset)
+		/** The size of the handle used to resize the drawer */
+		SLATE_ARGUMENT(float, ExpanderHandleSize)
+		/** Called when the target drawer size is changed by the user */
+		SLATE_EVENT(FOnSidebarDrawerTargetSizeChanged, OnDrawerTargetSizeChanged)
+		/** Called when the drawer gains focus */
+		SLATE_EVENT(FGenericSidebarDrawerWidgetDelegate, OnDrawerFocused)
+		/** Called when the drawer loses focus */
+		SLATE_EVENT(FGenericSidebarDrawerWidgetDelegate, OnDrawerFocusLost)
+		/** Called when the drawer is completely closed (i.e will be called once the close animation completes */
+		SLATE_EVENT(FGenericSidebarDrawerWidgetDelegate, OnDrawerClosed)
+	SLATE_END_ARGS()
+
+	~SSidebarDrawer();
+
+	void Construct(const FArguments& InArgs, const TSharedRef<FSidebarDrawer>& InDrawer, const ESidebarTabLocation InTabLocation);
+
+	/**
+	 * Opens the drawer.
+	 *
+	 * @param bInAnimateOpen Whether to play an animation when opening the drawer, defaults to true.
+	 */
+	void Open(bool bInAnimateOpen = true);
+
+	/** Begins an animation which closes the drawer. */
+	void Close();
+
+	/** @return True if the drawer is open. */
+	bool IsOpen() const;
+
+	/** @return True if the drawer is currently playing the close animation. */
+	bool IsClosing() const;
+
+	/** Sets the current size of the drawer, ignoring any open/close animation. */
+	void SetCurrentSize(const float InSize);
+
+	TSharedPtr<FSidebarDrawer> GetDrawer() const;
+
+protected:
+	//~ Begin SWidget
+	virtual bool SupportsKeyboardFocus() const override;
+	virtual FVector2D ComputeDesiredSize(const float InLayoutScaleMultiplier) const override;
+	virtual void OnArrangeChildren(const FGeometry& InAllottedGeometry, FArrangedChildren& ArrangedChildren) const override;
+	virtual FReply OnMouseButtonDown(const FGeometry& InAllottedGeometry, const FPointerEvent& InMouseEvent) override;
+	virtual FReply OnMouseButtonUp(const FGeometry& InAllottedGeometry, const FPointerEvent& InMouseEvent) override;
+	virtual FReply OnMouseMove(const FGeometry& InAllottedGeometry, const FPointerEvent& InMouseEvent) override;
+	virtual void OnMouseLeave(const FPointerEvent& InMouseEvent) override;
+	virtual FCursorReply OnCursorQuery(const FGeometry& InGeometry, const FPointerEvent& InCursorEvent) const override;
+	virtual int32 OnPaint(const FPaintArgs& InArgs, const FGeometry& InAllottedGeometry, const FSlateRect& InCullingRect
+		, FSlateWindowElementList& OutDrawElements, int32 InLayerId, const FWidgetStyle& InWidgetStyle, bool bInParentEnabled) const override;
+	//~ End SWidget
+
+	FGeometry GetRenderTransformedGeometry(const FGeometry& InAllottedGeometry) const;
+	
+	FGeometry GetResizeHandleGeometry(const FGeometry& InAllottedGeometry) const;
+	
+	EActiveTimerReturnType UpdateAnimation(const double InCurrentTime, const float InDeltaTime);
+	
+	void OnGlobalFocusChanging(const FFocusEvent& InFocusEvent
+		, const FWeakWidgetPath& InOldFocusedWidgetPath, const TSharedPtr<SWidget>& InOldFocusedWidget
+		, const FWidgetPath& InNewFocusedWidgetPath, const TSharedPtr<SWidget>& InNewFocusedWidget);
+	
+	void OnActiveTabChanged(const TSharedPtr<FSidebarDrawer>& InNewlyActivated, const TSharedPtr<FSidebarDrawer>& InPreviouslyActive);
+
+	TWeakPtr<FSidebarDrawer> DrawerWeak;
+	ESidebarTabLocation TabLocation = ESidebarTabLocation::Right;
+	
+	float MinDrawerSize = 0.f;
+	float MaxDrawerSize = 0.f;
+	float TargetDrawerSize = 0.f;
+	FVector2D ShadowOffset = FVector2D::ZeroVector;
+	float ExpanderHandleSize = 5.f;
+
+	FOnSidebarDrawerTargetSizeChanged OnDrawerTargetSizeChanged;
+	FGenericSidebarDrawerWidgetDelegate OnDrawerFocused;
+	FGenericSidebarDrawerWidgetDelegate OnDrawerFocusLost;
+	FGenericSidebarDrawerWidgetDelegate OnDrawerClosed;
+
+	FCurveSequence OpenCloseAnimation;
+	TSharedPtr<FActiveTimerHandle> OpenCloseTimer;
+
+	FThrottleRequest ResizeThrottleHandle;
+	FThrottleRequest AnimationThrottle;
+	
+	float CurrentSize = 0.f;
+	
+	bool bIsResizing = false;
+	bool bIsResizeHandleHovered = false;
+	float InitialSizeAtResize = 0.f;
+	FGeometry InitialResizeGeometry;
+};
diff --git a/Engine/Source/Developer/ToolWidgets/Private/Sidebar/SSidebarDrawerContent.cpp b/Engine/Source/Developer/ToolWidgets/Private/Sidebar/SSidebarDrawerContent.cpp
new file mode 100644
--- /dev/null
+++ b/Engine/Source/Developer/ToolWidgets/Private/Sidebar/SSidebarDrawerContent.cpp
@@ -0,0 +1,285 @@
+// Copyright Epic Games, Inc. All Rights Reserved.
+
+#include "SSidebarDrawerContent.h"
+#include "Framework/Application/SlateApplication.h"
+#include "Sidebar/ISidebarDrawerContent.h"
+#include "Sidebar/SidebarDrawer.h"
+#include "Widgets/Input/SCheckBox.h"
+#include "Widgets/Layout/SBox.h"
+#include "Widgets/Layout/SScrollBox.h"
+#include "Widgets/Layout/SWrapBox.h"
+#include "Widgets/SBoxPanel.h"
+#include "Widgets/Text/STextBlock.h"
+
+#define LOCTEXT_NAMESPACE "SSidebarDrawerContent"
+
+void SSidebarDrawerContent::Construct(const FArguments& InArgs, const TWeakPtr<FSidebarDrawer>& InOwnerDrawerWeak)
+{
+	OwnerDrawerWeak = InOwnerDrawerWeak;
+
+	ChildSlot
+	[
+		SNew(SVerticalBox)
+		+ SVerticalBox::Slot()
+		.AutoHeight()
+		.HAlign(HAlign_Fill)
+		[
+			SNew(SBox)
+			.Visibility_Lambda([this]()
+				{
+					return GetOrderedSections().Num() > 1 ? EVisibility::Visible : EVisibility::Collapsed;
+				})
+			[
+				SNew(SBox)
+				.Padding(FMargin(0.f, 4.f))
+				[
+					SAssignNew(ButtonBox, SWrapBox)
+					.HAlign(HAlign_Center)
+					.UseAllottedSize(true)
+					.InnerSlotPadding(FVector2D(4.f, 4.f))
+				]
+			]
+		]
+		+ SVerticalBox::Slot()
+		.FillHeight(1.f)
+		[
+			SAssignNew(ContentBox, SScrollBox)
+			.Orientation(Orient_Vertical)
+		]
+	];
+
+	BuildContent();
+}
+
+void SSidebarDrawerContent::BuildContent()
+{
+	const TSharedPtr<FSidebarDrawer> Drawer = OwnerDrawerWeak.Pin();
+	if (!Drawer.IsValid())
+	{
+		return;
+	}
+
+	TMap<FName, TSharedRef<ISidebarDrawerContent>> SortedContentDrawers;
+	SortedContentDrawers.Reserve(Drawer->ContentSections.Num());
+
+	ButtonBox->ClearChildren();
+	ContentBox->ClearChildren();
+
+	if (Drawer->ContentSections.IsEmpty())
+	{
+		return;
+	}
+
+	SortedContentDrawers = Drawer->ContentSections;
+	SortedContentDrawers.ValueSort(
+		[](const TSharedRef<ISidebarDrawerContent>& InA, const TSharedRef<ISidebarDrawerContent>& InB)
+		{
+			const int32 SortOrderA = InA->GetSortOrder();
+			const int32 SortOrderB = InB->GetSortOrder();
+
+			if (SortOrderA == SortOrderB)
+			{
+				static const FName General(TEXT("General"));
+				static const FName All(TEXT("All"));
+
+				const FName SectionNameA = InA->GetSectionId();
+				const FName SectionNameB = InB->GetSectionId();
+
+				// General first, All last, rest alphabetical
+				if (SectionNameA.IsEqual(General) || SectionNameB.IsEqual(All))
+				{
+					return true;
+				}
+				if (SectionNameA.IsEqual(All) || SectionNameB.IsEqual(General))
+				{
+					return false;
+				}
+
+				return SectionNameA.LexicalLess(SectionNameB);
+			}
+
+			return SortOrderA < SortOrderB;
+		});
+
+	const TArray<TSharedRef<ISidebarDrawerContent>> OriginalSectionNamesOrder = GetOrderedSections();
+	for (const TSharedRef<ISidebarDrawerContent>& Section : OriginalSectionNamesOrder)
+	{
+		const FName SectionName = Section->GetSectionId();
+
+		ButtonBox->AddSlot()
+			[
+				SNew(SBox)
+				.Padding(FMargin(0))
+				.HAlign(HAlign_Fill)
+				.VAlign(VAlign_Fill)
+				.Visibility(this, &SSidebarDrawerContent::GetSectionButtonVisibility, Section->AsWeak())
+				[
+					SNew(SCheckBox)
+					.Style(FAppStyle::Get(), TEXT("DetailsView.SectionButton"))
+					.OnCheckStateChanged(this, &SSidebarDrawerContent::OnSectionSelected, SectionName)
+					.IsChecked(this, &SSidebarDrawerContent::GetSectionCheckBoxState, SectionName)
+					[
+						SNew(STextBlock)
+						.TextStyle(FAppStyle::Get(), TEXT("SmallText"))
+						.Text(Section->GetSectionDisplayText())
+						.Justification(ETextJustify::Center)
+					]
+				]
+			];
+	}
+
+	for (const TPair<FName, TSharedRef<ISidebarDrawerContent>>& SectionPair : SortedContentDrawers)
+	{
+		AddContentSlot(SectionPair.Value, 1.f);
+	}
+
+	if (Drawer->SelectedContentSections.IsEmpty())
+	{
+		// Find first section that is visible
+		FName FoundSectionName;
+		for (const TPair<FName, TSharedRef<ISidebarDrawerContent>>& Section : Drawer->ContentSections)
+		{
+			if (Section.Value->ShouldShowSection())
+			{
+				FoundSectionName = Section.Value->GetSectionId();
+				break;
+			}
+		}
+
+		if (FoundSectionName != NAME_None)
+		{
+			Drawer->SelectedContentSections.Add(FoundSectionName);
+		}
+	}
+}
+
+void SSidebarDrawerContent::OnSectionSelected(const ECheckBoxState InCheckBoxState, const FName InSectionName)
+{
+	const TSharedPtr<FSidebarDrawer> Drawer = OwnerDrawerWeak.Pin();
+	if (!Drawer.IsValid())
+	{
+		return;
+	}
+
+	const FModifierKeysState ModifierKeysState = FSlateApplication::Get().GetModifierKeys();
+	const bool bIsModifierDown = ModifierKeysState.IsControlDown() || ModifierKeysState.IsShiftDown();
+
+	if (InCheckBoxState == ECheckBoxState::Checked)
+	{
+		if (bIsModifierDown)
+		{
+			Drawer->SelectedContentSections.Add(InSectionName);
+		}
+		else
+		{
+			Drawer->SelectedContentSections.Reset();
+			Drawer->SelectedContentSections.Add(InSectionName);
+		}
+	}
+	else
+	{
+		if (bIsModifierDown)
+		{
+			Drawer->SelectedContentSections.Remove(InSectionName);
+
+			// Force to always have a Selected Section, making it unable to de-select the last one
+			if (Drawer->SelectedContentSections.IsEmpty())
+			{
+				Drawer->SelectedContentSections.Add(InSectionName);
+			}
+		}
+		else
+		{
+			Drawer->SelectedContentSections.Reset();
+			Drawer->SelectedContentSections.Add(InSectionName);
+		}
+	}
+}
+
+bool SSidebarDrawerContent::IsSectionSelected(const FName InSectionName) const
+{
+	const TSharedPtr<FSidebarDrawer> Drawer = OwnerDrawerWeak.Pin();
+	return Drawer.IsValid() && Drawer->SelectedContentSections.Contains(InSectionName);
+}
+
+bool SSidebarDrawerContent::ShouldShowSection(const TWeakPtr<ISidebarDrawerContent>& InSectionWeak) const
+{
+	const TSharedPtr<ISidebarDrawerContent> Section = InSectionWeak.Pin();
+	return Section.IsValid() && Section->ShouldShowSection();
+}
+
+EVisibility SSidebarDrawerContent::GetSectionButtonVisibility(TWeakPtr<ISidebarDrawerContent> InSectionWeak) const
+{
+	return ShouldShowSection(InSectionWeak) ? EVisibility::SelfHitTestInvisible : EVisibility::Collapsed;
+}
+
+EVisibility SSidebarDrawerContent::GetSectionContentVisibility(const FName InSectionName, TWeakPtr<ISidebarDrawerContent> InSectionWeak) const
+{
+	return IsSectionSelected(InSectionName) && ShouldShowSection(InSectionWeak)
+		? EVisibility::SelfHitTestInvisible
+		: EVisibility::Collapsed;
+}
+
+ECheckBoxState SSidebarDrawerContent::GetSectionCheckBoxState(const FName InSectionName) const
+{
+	return IsSectionSelected(InSectionName) ? ECheckBoxState::Checked : ECheckBoxState::Unchecked;
+}
+
+TArray<TSharedRef<ISidebarDrawerContent>> SSidebarDrawerContent::GetOrderedSections() const
+{
+	const TSharedPtr<FSidebarDrawer> Drawer = OwnerDrawerWeak.Pin();
+	if (!Drawer.IsValid())
+	{
+		return {};
+	}
+
+	TArray<TSharedRef<ISidebarDrawerContent>> OrderedSections;
+
+	for (const TPair<FName, TSharedRef<ISidebarDrawerContent>>& SectionPair : Drawer->ContentSections)
+	{
+		const bool bContains = OrderedSections.ContainsByPredicate(
+			[&SectionPair](const TSharedRef<ISidebarDrawerContent>& InOther)
+			{
+				return SectionPair.Value->GetSectionId() == InOther->GetSectionId();
+			});
+		if (!bContains)
+		{
+			OrderedSections.Add(SectionPair.Value);
+		}
+	}
+
+	return MoveTemp(OrderedSections);
+}
+
+void SSidebarDrawerContent::AddContentSlot(const TSharedRef<ISidebarDrawerContent>& InDrawerContent, const TOptional<float>& InFillSize)
+{
+	const float BottomPadding = 2.f;
+	const FName SectionName = InDrawerContent->GetSectionId();
+
+	const TSharedRef<SBox> ContentWidget = SNew(SBox)
+		.Visibility(this, &SSidebarDrawerContent::GetSectionContentVisibility, SectionName, InDrawerContent->AsWeak())
+		[
+			InDrawerContent->CreateContentWidget()
+		];
+
+	if (InFillSize.IsSet())
+	{
+		ContentBox->AddSlot()
+			.FillSize(InFillSize.Get(1.f))
+			.Padding(0.f, 0.f, 0.f, BottomPadding)
+			[
+				ContentWidget
+			];
+	}
+	else
+	{
+		ContentBox->AddSlot()
+			.AutoSize()
+			.Padding(0.f, 0.f, 0.f, BottomPadding)
+			[
+				ContentWidget
+			];
+	}
+}
+
+#undef LOCTEXT_NAMESPACE 
diff --git a/Engine/Source/Developer/ToolWidgets/Private/Sidebar/SSidebarDrawerContent.h b/Engine/Source/Developer/ToolWidgets/Private/Sidebar/SSidebarDrawerContent.h
new file mode 100644
--- /dev/null
+++ b/Engine/Source/Developer/ToolWidgets/Private/Sidebar/SSidebarDrawerContent.h
@@ -0,0 +1,48 @@
+// Copyright Epic Games, Inc. All Rights Reserved.
+
+#pragma once
+
+#include "Containers/Array.h"
+#include "Styling/SlateTypes.h"
+#include "Templates/SharedPointer.h"
+#include "UObject/NameTypes.h"
+#include "Widgets/SCompoundWidget.h"
+
+class FSidebarDrawer;
+class ISidebarDrawerContent;
+class SScrollBox;
+class SWrapBox;
+
+/** Handles drawer multi-section display */
+class SSidebarDrawerContent : public SCompoundWidget
+{
+public:
+	SLATE_BEGIN_ARGS(SSidebarDrawerContent) {}
+	SLATE_END_ARGS()
+
+	void Construct(const FArguments& InArgs, const TWeakPtr<FSidebarDrawer>& InOwnerDrawerWeak);
+
+	void BuildContent();
+
+private:
+	void OnSectionSelected(const ECheckBoxState InCheckBoxState, const FName InSectionName);
+
+	bool IsSectionSelected(const FName InSectionName) const;
+
+	bool ShouldShowSection(const TWeakPtr<ISidebarDrawerContent>& InSectionWeak) const;
+
+	EVisibility GetSectionButtonVisibility(TWeakPtr<ISidebarDrawerContent> InSectionWeak) const;
+	
+	EVisibility GetSectionContentVisibility(const FName InSectionName, TWeakPtr<ISidebarDrawerContent> InSectionWeak) const;
+
+	ECheckBoxState GetSectionCheckBoxState(const FName InSectionName) const;
+
+	TArray<TSharedRef<ISidebarDrawerContent>> GetOrderedSections() const;
+
+	void AddContentSlot(const TSharedRef<ISidebarDrawerContent>& InDrawerContent, const TOptional<float>& InFillSize = TOptional<float>());
+
+	TWeakPtr<FSidebarDrawer> OwnerDrawerWeak;
+
+	TSharedPtr<SWrapBox> ButtonBox;
+	TSharedPtr<SScrollBox> ContentBox;
+};
diff --git a/Engine/Source/Developer/ToolWidgets/Private/Sidebar/SidebarButtonMenuContext.h b/Engine/Source/Developer/ToolWidgets/Private/Sidebar/SidebarButtonMenuContext.h
new file mode 100644
--- /dev/null
+++ b/Engine/Source/Developer/ToolWidgets/Private/Sidebar/SidebarButtonMenuContext.h
@@ -0,0 +1,37 @@
+﻿// Copyright Epic Games, Inc. All Rights Reserved.
+
+#pragma once
+
+#include "Templates/SharedPointer.h"
+#include "SidebarButtonMenuContext.generated.h"
+
+class FSidebarDrawer;
+class SSidebar;
+
+UCLASS()
+class USidebarButtonMenuContext : public UObject
+{
+	GENERATED_BODY()
+
+public:
+	void Init(const TWeakPtr<SSidebar>& InSidebarWeak, const TWeakPtr<FSidebarDrawer>& InDrawerWeak)
+	{
+		SidebarWeak = InSidebarWeak;
+		DrawerWeak = InDrawerWeak;
+	}
+
+	TSharedPtr<SSidebar> GetSidebarWidget() const
+	{
+		return SidebarWeak.Pin();
+	}
+
+	TSharedPtr<FSidebarDrawer> GetDrawer() const
+	{
+		return DrawerWeak.Pin();
+	}
+
+protected:
+	TWeakPtr<SSidebar> SidebarWeak;
+
+	TWeakPtr<FSidebarDrawer> DrawerWeak;
+};
diff --git a/Engine/Source/Developer/ToolWidgets/Private/Sidebar/SidebarDrawer.h b/Engine/Source/Developer/ToolWidgets/Private/Sidebar/SidebarDrawer.h
new file mode 100644
--- /dev/null
+++ b/Engine/Source/Developer/ToolWidgets/Private/Sidebar/SidebarDrawer.h
@@ -0,0 +1,75 @@
+// Copyright Epic Games, Inc. All Rights Reserved.
+
+#pragma once
+
+#include "Framework/SlateDelegates.h"
+#include "Sidebar/SidebarDrawerConfig.h"
+#include "Templates/SharedPointer.h"
+#include "UObject/NameTypes.h"
+
+class ISidebarDrawerContent;
+class SSidebarDrawer;
+class SWidget;
+
+DECLARE_DELEGATE_OneParam(FOnSidebarDrawerOpened, const FName /*InUniqueId*/);
+DECLARE_DELEGATE_OneParam(FOnSidebarDrawerClosed, const FName /*InUniqueId*/);
+
+class FSidebarDrawer
+{
+public:
+	FSidebarDrawer(FSidebarDrawerConfig&& InDrawerConfig)
+		: Config(MoveTemp(InDrawerConfig))
+		, SizeCoefficient(Config.DefaultSizeCoefficient)
+	{}
+
+	bool operator==(const FName InOtherId) const
+	{
+		return Config.UniqueId == InOtherId;
+	}
+
+	bool operator!=(const FName InOtherId) const
+	{
+		return Config.UniqueId != InOtherId;
+	}
+
+	bool operator==(const FSidebarDrawer& InOther) const
+	{
+		return Config == InOther.Config;
+	}
+
+	bool operator!=(const FSidebarDrawer& InOther) const
+	{
+		return Config != InOther.Config;
+	}
+
+	FName GetUniqueId() const
+	{
+		return Config.UniqueId;
+	}
+
+	FSidebarDrawerConfig Config;
+
+	/** Tab button widget for this drawer. */
+	TSharedPtr<SWidget> ButtonWidget;
+
+	/** Tab sliding drawer widget that contains the content. */
+	TSharedPtr<SSidebarDrawer> DrawerWidget;
+
+	/** The content widget contains the section widgets. */
+	TSharedPtr<SWidget> ContentWidget;
+
+	TMap<FName, TSharedRef<ISidebarDrawerContent>> ContentSections;
+	TSet<FName> SelectedContentSections;
+
+	bool bIsOpen = false;
+	bool bIsPinned = false;
+	bool bIsDocked = false;
+
+	bool bDisablePin = false;
+	bool bDisableDock = false;
+
+	float SizeCoefficient;
+
+	FOnSidebarDrawerOpened DrawerOpenedDelegate;
+	FOnSidebarDrawerClosed DrawerClosedDelegate;
+};
diff --git a/Engine/Source/Developer/ToolWidgets/Public/Sidebar/ISidebarDrawerContent.h b/Engine/Source/Developer/ToolWidgets/Public/Sidebar/ISidebarDrawerContent.h
new file mode 100644
--- /dev/null
+++ b/Engine/Source/Developer/ToolWidgets/Public/Sidebar/ISidebarDrawerContent.h
@@ -0,0 +1,33 @@
+// Copyright Epic Games, Inc. All Rights Reserved.
+
+#pragma once
+
+#include "Templates/SharedPointer.h"
+
+class FName;
+class FText;
+class SWidget;
+
+class ISidebarDrawerContent : public TSharedFromThis<ISidebarDrawerContent>
+{
+public:
+	virtual ~ISidebarDrawerContent() {}
+
+	/** Gets the unique Id used to identify the sidebar drawer. */
+	virtual FName GetUniqueId() const = 0;
+
+	/** Gets the section unique Id used to identify the drawer section. */
+	virtual FName GetSectionId() const = 0;
+
+	/** Gets the section text to display on the section button. */
+	virtual FText GetSectionDisplayText() const = 0;
+
+	/** Returns true if this section should be displayed. */
+	virtual bool ShouldShowSection() const { return true; }
+
+	/** Gets the section sort order. Lower values will be put at the above higher values. */
+	virtual int32 GetSortOrder() const { return 0; }
+
+	/** Constructs the content widget for a sidebar drawer. */
+	virtual TSharedRef<SWidget> CreateContentWidget() = 0;
+};
diff --git a/Engine/Source/Developer/ToolWidgets/Public/Sidebar/SSidebar.h b/Engine/Source/Developer/ToolWidgets/Public/Sidebar/SSidebar.h
new file mode 100644
--- /dev/null
+++ b/Engine/Source/Developer/ToolWidgets/Public/Sidebar/SSidebar.h
@@ -0,0 +1,290 @@
+// Copyright Epic Games, Inc. All Rights Reserved.
+
+#pragma once
+
+#include "SidebarDrawerConfig.h"
+#include "Widgets/DeclarativeSyntaxSupport.h"
+#include "Widgets/Layout/SSplitter.h"
+#include "Widgets/SCompoundWidget.h"
+
+class FExtender;
+class FSidebarDrawer;
+class ISidebarDrawerContent;
+class SBox;
+class SOverlay;
+class SScrollBox;
+class SSidebarButton;
+class SSidebarDrawer;
+class STabDrawer;
+class SVerticalBox;
+class UToolMenu;
+struct FTabId;
+
+/**
+ * The direction that a tab drawer opens relative to the location of the sidebar.
+ * NOTE: Effort has been made to support top and bottom sidebar locations, but this has not been thoroughly tested
+ *   and ironed out because there is currently no use case.
+ */
+enum class ESidebarTabLocation : uint8
+{
+	/** Open from left to right */
+	Left,
+	/** Open from right to left */
+	Right,
+	/** Open from bottom to top */
+	Top,
+	/** Open from top to bottom */
+	Bottom
+};
+
+DECLARE_DELEGATE_OneParam(FOnSidebarDrawerDockStateChanged, const FName /*InDrawerId*/);
+
+/**
+ * Static sidebar tab widget that cannot be dragged or moved to a different location. Multiple drawers can be registered
+ * to a single sidebar and each drawer can have its own sections, each of which can be displayed single, in combination,
+ * or all together through buttons at the top of the drawer.
+ */
+class TOOLWIDGETS_API SSidebar : public SCompoundWidget
+{
+public:
+	static constexpr float MinTabButtonSize = 100.f;
+	static constexpr float MaxTabButtonSize = 200.f;
+	static constexpr float TabButtonThickness = 25.f;
+
+	SLATE_BEGIN_ARGS(SSidebar)
+		: _TabLocation(ESidebarTabLocation::Right)
+		, _HideWhenDocked(false)
+		, _AlwaysUseMaxButtonSize(false)
+		, _DisablePin(false)
+		, _DisableDock(false)
+	{}
+		/** The direction that a tab drawer opens relative to the location of the sidebar. */
+		SLATE_ARGUMENT(ESidebarTabLocation, TabLocation)
+		/** Hides the sidebar when a drawer is docked. NOTE: Must provide a way to manually undock the drawer to restore the sidebar visibility. */
+		SLATE_ARGUMENT(bool, HideWhenDocked)
+		/** Forces the sidebar tab buttons to always be a uniform size of max. */
+		SLATE_ARGUMENT(bool, AlwaysUseMaxButtonSize)
+		/** Disables the ability to pin a drawer. */
+		SLATE_ARGUMENT(bool, DisablePin)
+		/** Disables the ability to dock a drawer. */
+		SLATE_ARGUMENT(bool, DisableDock)
+		/** Event triggered when a drawers dock state changes. */
+		SLATE_EVENT(FOnSidebarDrawerDockStateChanged, OnDockStateChanged)
+	SLATE_END_ARGS()
+
+	virtual ~SSidebar() override;
+
+	/**
+	 * Constructs the sidebar widget.
+	 * 
+	 * @param InArgs Widget construction arguments
+	 * @param InDrawersOverlay Overlay widget used to display the animating drawer
+	 * @param InDockLocation Parent widget that will contain the drawer content widget when docked
+	 */
+	void Construct(const FArguments& InArgs, const TSharedRef<SOverlay>& InDrawersOverlay, const TSharedRef<SBox>& InDockLocation);
+
+	/**
+	 * Registers and displays a new drawer in the sidebar.
+	 * 
+	 * @param InDrawerConfig Configuration info for the new drawer
+	 * 
+	 * @return True if the new drawer registration was successful.
+	 */
+	bool RegisterDrawer(FSidebarDrawerConfig&& InDrawerConfig);
+
+	/**
+	 * Unregisters and removes a drawer from the sidebar.
+	 *
+	 * @param InDrawerId Unique drawer Id to unregister
+	 * 
+	 * @return True if the drawer removal was successful.
+	 */
+	bool UnregisterDrawer(const FName InDrawerId);
+
+	/**
+	 * Checks if a drawer exists in the sidebar.
+	 * 
+	 * @param InDrawerId Unique drawer Id to check exists
+	 * 
+	 * @return True if the sidebar contains a drawer with the specified Id.
+	 */
+	bool ContainsDrawer(const FName InDrawerId) const;
+
+	/** @return The number of drawers that exist in the sidebar. */
+	int32 GetDrawerCount() const;
+
+	/**
+	 * Registers and displays a new drawer section in the sidebar.
+	 * 
+	 * @param InDrawerId Unique drawer Id to register
+	 * @param InSection Drawer content interface for the section
+	 * 
+	 * @return True if the new drawer section registration was successful.
+	 */
+	bool RegisterDrawerSection(const FName InDrawerId, const TSharedPtr<ISidebarDrawerContent>& InSection);
+
+	/**
+	 * Unregisters and removes a drawer section from the sidebar.
+	 * 
+	 * @param InDrawerId Unique drawer Id that contains the section to unregister
+	 * @param InSectionId Unique drawer section Id to unregister
+	 * 
+	 * @return True if the drawer removal was successful.
+	 */
+	bool UnregisterDrawerSection(const FName InDrawerId, const FName InSectionId);
+
+	/**
+	 * Checks if a drawer section exists within a sidebar drawer.
+	 * 
+	 * @param InDrawerId Unique Id of the drawer to look for the section in
+	 * @param InSectionId Unique drawer section Id to check exists
+	 * 
+	 * @return True if the sidebar contains a drawer with the specified Id.
+	 */
+	bool ContainsDrawerSection(const FName InDrawerId, const FName InSectionId) const;
+
+	/**
+	 * Attempt to open a specific drawer in the sidebar.
+	 *
+	 * @param InDrawerId Unique Id of the drawer to attempt to open
+	 * 
+	 * @return True if the drawer exists in this sidebar and was opened.
+	 */
+	bool TryOpenDrawer(const FName InDrawerId);
+
+	/** Closes any drawers that are open. */
+	void CloseAllDrawers();
+
+	/** @return True if the sidebar has any drawer that is opened. */
+	bool HasDrawerOpened() const;
+
+	/**
+	 * Checks if a drawer is opened.
+	 * 
+	 * @param InDrawerId Unique Id of the drawer to check
+	 * 
+	 * @return True if the specified drawer is currently opened.
+	 */
+	bool IsDrawerOpened(const FName InDrawerId) const;
+
+	/** @return The unique drawer Id that is currently open. */
+	FName GetOpenedDrawerId() const;
+
+	/** @return True if the sidebar has any drawer that is pinned. */
+	bool HasDrawerPinned() const;
+
+	/**
+	 * Checks if a drawer is pinned.
+	 * 
+	 * @param InDrawerId Unique Id of the drawer to check
+	 * 
+	 * @return True if the specified drawer is currently pinned.
+	 */
+	bool IsDrawerPinned(const FName InDrawerId) const;
+
+	/**
+	 * Pins a drawer so it stays open even when focus is lost.
+	 * 
+	 * @param InDrawerId Unique Id of the drawer to pin
+	 * @param bInIsPinned New pin state to set
+	 */
+	void SetDrawerPinned(const FName InDrawerId, const bool bInIsPinned);
+
+	/** @return True if the sidebar has any drawer that is docked. */
+	bool HasDrawerDocked() const;
+
+	/** @return True if the specified drawer Id is docked. */
+	bool IsDrawerDocked(const FName InDrawerId) const;
+
+	/**
+	 * Docks a drawer so it embeds itself into the content.
+	 * 
+	 * @param InDrawerId Unique Id of the drawer to dock
+	 * @param bInIsDocked New dock state to set
+	 */
+	void SetDrawerDocked(const FName InDrawerId, const bool bInIsDocked);
+
+	/** Undocks any drawers that are docked. */
+	void UndockAllDrawers();
+
+	/**
+	 * Helper function to update a splitter slot size based on a drawers state.
+	 * Sets the slots resize-ability, sizing rule, size value.
+	 * 
+	 * @param InDrawerId Unique Id of the drawer to update for
+	 * @param InSlot The splitter slot to modify the size of
+	 * @param bInAutoUndock Automatically undocks the drawer if the size is below the threshold
+	 * @param InDefaultDockPercent Default size co-efficient to reset the drawer size to
+	 */
+	void UpdateDockedSplitterSlot(const FName InDrawerId, SSplitter::FSlot* const InSlot, const bool bInAutoUndock = true
+		, const float InDefaultDockPercent = FSidebarDrawerConfig::DefaultSizeCoefficient);
+
+	/** @return True if the sidebar is set to animate horizontally. */
+	bool IsHorizontal() const;
+
+	/** @return True if the sidebar is set to animate vertically. */
+	bool IsVertical() const;
+
+private:
+	void OnTabDrawerButtonPressed(const TSharedRef<FSidebarDrawer>& InDrawer);
+	void OnDrawerTabPinToggled(const TSharedRef<FSidebarDrawer>& InDrawer, const bool bIsPinned);
+	void OnDrawerTabDockToggled(const TSharedRef<FSidebarDrawer>& InDrawer, const bool bIsDocked);
+	void OnTabDrawerFocusLost(const TSharedRef<SSidebarDrawer>& InDrawerWidget);
+	void OnTabDrawerClosed(const TSharedRef<SSidebarDrawer>& InDrawerWidget);
+	void OnDrawerTargetSizeChanged(const TSharedRef<SSidebarDrawer>& InDrawerWidget, const float InNewSize);
+	TSharedRef<SWidget> OnGetTabDrawerContextMenuWidget(TSharedRef<FSidebarDrawer> InDrawer);
+	void BuildOptionsMenu(UToolMenu* const InMenu);
+
+	/** Removes a single drawer for a specified tab from this sidebar. Removal is done instantly not waiting for any close animation. */
+	void RemoveDrawer(const TSharedRef<FSidebarDrawer>& InDrawer);
+
+	/** Removes all drawers instantly (including drawers for pinned tabs). */
+	void RemoveAllDrawers();
+
+	EActiveTimerReturnType OnOpenPendingDrawerTimer(const double CurrentTime, const float DeltaTime);
+	void OpenDrawerNextFrame(const TSharedRef<FSidebarDrawer>& InDrawer, const bool bAnimateOpen);
+	void OpenDrawerInternal(const TSharedRef<FSidebarDrawer>& InDrawer, const bool bAnimateOpen);
+	void CloseDrawerInternal(const TSharedRef<FSidebarDrawer>& InDrawer);
+
+	/** Reopens the pinned tab only if there are no other open drawers. This should be used to bring pinned tabs back after other tabs lose focus/are closed. */
+	void SummonPinnedTabIfNothingOpened();
+
+	/** Updates the appearance of open drawers. */
+	void UpdateDrawerAppearance();
+
+	TSharedPtr<FSidebarDrawer> FindDrawer(const FName InDrawerId) const;
+
+	/** Returns the first tab in this sidebar that is marked pinned. */
+	TSharedPtr<FSidebarDrawer> FindFirstPinnedTab() const;
+
+	/** Returns the tab for the last-opened drawer that is still open, excluding any drawers that are in the process of closing. */
+	TSharedPtr<FSidebarDrawer> GetForegroundTab() const;
+
+	/** Returns the drawer for the given tab if it's open. */
+	TSharedPtr<SSidebarDrawer> FindOpenedDrawer(const TSharedRef<FSidebarDrawer>& InDrawer) const;
+
+	TWeakPtr<SOverlay> DrawersOverlayWeak;
+	TWeakPtr<SBox> DockLocationWeak;
+
+	ESidebarTabLocation TabLocation = ESidebarTabLocation::Right;
+	bool bHideWhenDocked = false;
+	bool bAlwaysUseMaxButtonSize = false;
+	bool bDisablePin = false;
+	bool bDisableDock = false;
+	FOnSidebarDrawerDockStateChanged OnDockStateChanged;
+
+	TSharedPtr<SScrollBox> TabContainer;
+
+	TArray<TSharedRef<FSidebarDrawer>> DrawerTabs;
+
+	/** Generally speaking one drawer is only ever open at once but we animate any previous drawer closing so there could be more than one while an animation is playing. */
+	TArray<TSharedRef<SSidebarDrawer>> OpenedDrawers;
+
+	TWeakPtr<FSidebarDrawer> PendingTabToOpen;
+	bool bAnimatePendingTabOpen = false;
+	TSharedPtr<FActiveTimerHandle> OpenPendingDrawerTimerHandle;
+
+	TArray<TSharedRef<FSidebarDrawer>> PinnedDrawerTabs;
+
+	TSharedPtr<FSidebarDrawer> DockedDrawerTab;
+};
diff --git a/Engine/Source/Developer/ToolWidgets/Public/Sidebar/SidebarDrawerConfig.h b/Engine/Source/Developer/ToolWidgets/Public/Sidebar/SidebarDrawerConfig.h
new file mode 100644
--- /dev/null
+++ b/Engine/Source/Developer/ToolWidgets/Public/Sidebar/SidebarDrawerConfig.h
@@ -0,0 +1,61 @@
+// Copyright Epic Games, Inc. All Rights Reserved.
+
+#pragma once
+
+#include "Containers/Array.h"
+#include "Internationalization/Text.h"
+#include "Misc/Attribute.h"
+#include "Templates/SharedPointer.h"
+#include "UObject/NameTypes.h"
+
+class SWidget;
+struct FSlateBrush;
+
+/** Configuration information used to register a sidebar drawer. */
+struct TOOLWIDGETS_API FSidebarDrawerConfig
+{
+	static constexpr float DefaultSizeCoefficient = 0.25f;
+
+	bool operator==(const FName InOtherId) const
+	{
+		return UniqueId == InOtherId;
+	}
+
+	bool operator!=(const FName InOtherId) const
+	{
+		return UniqueId != InOtherId;
+	}
+
+	bool operator==(const FSidebarDrawerConfig& InOther) const
+	{
+		return UniqueId == InOther.UniqueId;
+	}
+
+	bool operator!=(const FSidebarDrawerConfig& InOther) const
+	{
+		return UniqueId != InOther.UniqueId;
+	}
+
+	/** Unique Id used to identify the drawer */
+	FName UniqueId;
+
+	/** Text to display on the drawer tab button */
+	TAttribute<FText> ButtonText;
+
+	/** ToolTip text to display for the drawer tab button */
+	TAttribute<FText> ToolTipText;
+
+	/** Icon to display for the drawer tab button */
+	TAttribute<const FSlateBrush*> Icon;
+
+	/** True if the drawers initial state should be docked. */
+	bool bInitiallyDocked = false;
+
+	/** Size co-efficient from 0 to 1 to use for the size of the docked content. */
+	float SizeCoefficient = DefaultSizeCoefficient;
+
+	/** Optional content widget to use instead of registering sections. */
+	TSharedPtr<SWidget> OverrideContentWidget;
+
+	// @TODO: events for pin state, dock state, section state changes?
+};
