[ToolWidgets] Fix SSidebar LOCTEXT duplicate key.

diff --git a/Engine/Source/Developer/ToolWidgets/Private/Sidebar/SSidebar.cpp b/Engine/Source/Developer/ToolWidgets/Private/Sidebar/SSidebar.cpp
--- a/Engine/Source/Developer/ToolWidgets/Private/Sidebar/SSidebar.cpp
+++ b/Engine/Source/Developer/ToolWidgets/Private/Sidebar/SSidebar.cpp
@@ -371,8 +371,8 @@ void SSidebar::BuildOptionsMenu(UToolMenu* const InMenu)
 	else
 	{
 		Section.AddMenuEntry(TEXT("Dock"),
-			LOCTEXT("UndockLabel", "Dock"),
-			LOCTEXT("UndockToolTip", "Docks the drawer"),
+			LOCTEXT("DockLabel", "Dock"),
+			LOCTEXT("DockToolTip", "Docks the drawer"),
 			FSlateIcon(),
 			FUIAction(FExecuteAction::CreateSP(this, &SSidebar::SetDrawerDocked, Drawer->GetUniqueId(), true)));
 	}
