[ToolWidgets] Restored previous base classes of buttons, and refactored to keep shared functionality/styling

#rnx

diff --git a/Engine/Source/Developer/ToolWidgets/Private/SActionButton.cpp b/Engine/Source/Developer/ToolWidgets/Private/SActionButton.cpp
--- a/Engine/Source/Developer/ToolWidgets/Private/SActionButton.cpp
+++ b/Engine/Source/Developer/ToolWidgets/Private/SActionButton.cpp
@@ -4,11 +4,10 @@
 
 #include "Styling/StyleColors.h"
 #include "ToolWidgetsStyle.h"
-#include "ToolWidgetsStylePrivate.h"
+#include "ToolWidgetsUtilitiesPrivate.h"
 #include "Widgets/Images/SImage.h"
 #include "Widgets/Input/SButton.h"
 #include "Widgets/Layout/SSpacer.h"
-#include "Widgets/Layout/SWidgetSwitcher.h"
 
 namespace UE::ToolWidgets::Private
 {
@@ -17,9 +16,10 @@ namespace UE::ToolWidgets::Private
 		static TMap<EActionButtonType, const FName> Lookup = {
 			{ EActionButtonType::Default, TEXT("ActionButton") },
 			{ EActionButtonType::Simple, TEXT("SimpleComboButton") },
+			{ EActionButtonType::Primary, TEXT("PrimaryButton") },
 			{ EActionButtonType::Positive, TEXT("PositiveActionButton") },
-			{ EActionButtonType::Warning, TEXT("NegativeActionButton") },
-			{ EActionButtonType::Error, TEXT("NegativeActionButton") },
+			{ EActionButtonType::Warning, TEXT("NegativeActionButton.Warning") },
+			{ EActionButtonType::Error, TEXT("NegativeActionButton.Error") },
 		};
 
 		if (const FName* FoundName = Lookup.Find(ActionButtonType))
@@ -60,14 +60,16 @@ void SActionButton::Construct(const FArguments& InArgs)
 	ComboButtonStyle = InArgs._ComboButtonStyle ? InArgs._ComboButtonStyle : &ActionButtonStyle->ComboButtonStyle;
 	TextBlockStyle = InArgs._TextBlockStyle ? InArgs._TextBlockStyle : &ActionButtonStyle->TextBlockStyle;
 
-	const EHorizontalAlignment HorizontalContentAlignment = InArgs._HorizontalContentAlignment.IsSet()
+	const EHorizontalAlignment HorizontalContentAlignment =
+		InArgs._HorizontalContentAlignment.IsSet()
 		? InArgs._HorizontalContentAlignment.Get(EHorizontalAlignment::HAlign_Center)
 		: static_cast<EHorizontalAlignment>(ActionButtonStyle->HorizontalContentAlignment);
 
 	const bool bHasDownArrow = InArgs._HasDownArrow.IsSet() ? InArgs._HasDownArrow.Get(false) : ActionButtonStyle->bHasDownArrow;
 
 	// Check for widget level override, then style override, otherwise unset
-	const TAttribute<const FSlateBrush*> Icon = InArgs._Icon.IsSet()
+	const TAttribute<const FSlateBrush*> Icon =
+		InArgs._Icon.IsSet()
 		? InArgs._Icon.Get()
 		: ActionButtonStyle->IconBrush.IsSet()
 		? &ActionButtonStyle->IconBrush.GetValue()
@@ -76,7 +78,8 @@ void SActionButton::Construct(const FArguments& InArgs)
 	const bool bHasIcon = Icon.Get() || Icon.IsBound();
 
 	// Check for widget level override, then style override, otherwise get from ActionButtonType
-	const TAttribute<FSlateColor> IconColorAndOpacity = InArgs._IconColorAndOpacity.IsSet()
+	const TAttribute<FSlateColor> IconColorAndOpacity =
+		InArgs._IconColorAndOpacity.IsSet()
 		? InArgs._IconColorAndOpacity.Get()
 		: ActionButtonStyle->IconColorAndOpacity.IsSet()
 		? TAttribute<FSlateColor>(ActionButtonStyle->IconColorAndOpacity.GetValue())
@@ -84,64 +87,22 @@ void SActionButton::Construct(const FArguments& InArgs)
 
 	TAttribute<FText> Text = InArgs._Text;
 
-	static constexpr float DefaultIconHeight = UE::ToolWidgets::Private::FToolWidgetsStylePrivate::FActionButton::DefaultIconHeight;
-	static constexpr float IconTextPadding = UE::ToolWidgets::Private::FToolWidgetsStylePrivate::FActionButton::DefaultIconLabelSpacing;
-
-	const TSharedRef<SHorizontalBox> ButtonContentContainer = SNew(SHorizontalBox)
-		+ SHorizontalBox::Slot()
-		.HAlign(HAlign_Center)
-		.VAlign(VAlign_Center)
-		.AutoWidth()
-		.Padding(0)
-		[
-			SNew(SWidgetSwitcher)
-			.WidgetIndex(bHasIcon ? 1 : 0)
-
-			+ SWidgetSwitcher::Slot()
-			[
-				SNew(SSpacer)
-				.Size(FVector2D{ 0, DefaultIconHeight })
-			]
-
-			+ SWidgetSwitcher::Slot()
-			[
-				SNew(SImage)
-				.Image(Icon)
-				.ColorAndOpacity(IconColorAndOpacity)
-				.Visibility(bHasIcon ? EVisibility::HitTestInvisible : EVisibility::Collapsed)
-			]
-		]
-
-		+ SHorizontalBox::Slot()
-		.VAlign(VAlign_Center)
-		.FillWidth(1.0f)
-		.Padding(FMargin(bHasIcon ? IconTextPadding : 0, 0, 0, 0))
-		[
-			SNew(STextBlock)
-			.TextStyle(TextBlockStyle)
-			.Text(InArgs._Text)
-			.Visibility_Lambda([Text]()
-			{
-				return Text.Get(FText::GetEmpty()).IsEmpty()
-					? EVisibility::Collapsed
-					: EVisibility::Visible;
-			})
-		];
+	const TSharedRef<SWidget> ButtonContent =
+		UE::ToolWidgets::Private::ActionButton::MakeButtonContent(
+			Icon,
+			IconColorAndOpacity,
+			Text,
+			TextBlockStyle);
 
 	// Treated as a ComboButton if OnClicked is not bound
 	const bool bIsComboButton = !InArgs._OnClicked.IsBound();
 
 	if (bIsComboButton)
 	{
-		static const FMargin DefaultComboButtonContentPadding = FMargin(
-			UE::ToolWidgets::Private::FToolWidgetsStylePrivate::FActionButton::DefaultHorizontalPadding,
-			UE::ToolWidgets::Private::FToolWidgetsStylePrivate::FActionButton::DefaultVerticalPadding);
-
-		const TAttribute<FMargin> ComboButtonContentPadding = InArgs._ButtonContentPadding.IsSet()
-			? InArgs._ButtonContentPadding.Get(DefaultComboButtonContentPadding)
-			: ActionButtonStyle->ComboButtonContentPadding.IsSet()
-			? ActionButtonStyle->ComboButtonContentPadding.GetValue()
-			: ActionButtonStyle->ComboButtonStyle.ContentPadding;
+		const TAttribute<FMargin> ComboButtonContentPadding =
+			InArgs._ButtonContentPadding.IsSet()
+			? InArgs._ButtonContentPadding.Get()
+			: ActionButtonStyle->GetComboButtonContentPadding();
 
 		ChildSlot
 		[
@@ -156,7 +117,7 @@ void SActionButton::Construct(const FArguments& InArgs)
 			.VAlign(VAlign_Center)
 			.ButtonContent()
 			[
-				ButtonContentContainer
+				ButtonContent
 			]
 			.MenuContent()
 			[
@@ -169,15 +130,10 @@ void SActionButton::Construct(const FArguments& InArgs)
 	}
 	else
 	{
-		static const FMargin DefaultButtonContentPadding = FMargin(
-			UE::ToolWidgets::Private::FToolWidgetsStylePrivate::FActionButton::DefaultHorizontalPadding,
-			UE::ToolWidgets::Private::FToolWidgetsStylePrivate::FActionButton::DefaultVerticalPadding);
-
-		const TAttribute<FMargin> ButtonContentPadding = InArgs._ButtonContentPadding.IsSet()
-			? InArgs._ButtonContentPadding.Get(DefaultButtonContentPadding)
-			: ActionButtonStyle->ButtonContentPadding.IsSet()
-			? ActionButtonStyle->ButtonContentPadding.GetValue()
-			: DefaultButtonContentPadding;
+		const TAttribute<FMargin> ButtonContentPadding =
+			InArgs._ButtonContentPadding.IsSet()
+			? InArgs._ButtonContentPadding.Get()
+			: ActionButtonStyle->GetButtonContentPadding();
 
 		ChildSlot
 		[
@@ -190,7 +146,7 @@ void SActionButton::Construct(const FArguments& InArgs)
 			.VAlign(VAlign_Center)
 			.OnClicked(InArgs._OnClicked)
 			[
-				ButtonContentContainer
+				ButtonContent
 			]
 		];
 	}
diff --git a/Engine/Source/Developer/ToolWidgets/Private/SPrimaryButton.cpp b/Engine/Source/Developer/ToolWidgets/Private/SPrimaryButton.cpp
--- a/Engine/Source/Developer/ToolWidgets/Private/SPrimaryButton.cpp
+++ b/Engine/Source/Developer/ToolWidgets/Private/SPrimaryButton.cpp
@@ -2,21 +2,45 @@
 
 #include "SPrimaryButton.h"
 
-#include "SActionButton.h"
 #include "ToolWidgetsSlateTypes.h"
 #include "ToolWidgetsStyle.h"
+#include "ToolWidgetsUtilitiesPrivate.h"
 #include "Widgets/SBoxPanel.h"
 #include "Widgets/Images/SImage.h"
 #include "Widgets/Text/STextBlock.h"
 
 void SPrimaryButton::Construct(const FArguments& InArgs)
 {
-	ChildSlot
-	[
-		SNew(SActionButton)
-		.ActionButtonStyle(&UE::ToolWidgets::FToolWidgetsStyle::Get().GetWidgetStyle<FActionButtonStyle>("PrimaryButton"))
-		.Text(InArgs._Text)
-		.Icon(InArgs._Icon)
+	const FActionButtonStyle* ActionButtonStyle = &UE::ToolWidgets::FToolWidgetsStyle::Get().GetWidgetStyle<FActionButtonStyle>("PrimaryButton");
+
+	// Check for widget level override, then style override, otherwise unset
+	const TAttribute<const FSlateBrush*> Icon = InArgs._Icon.IsSet()
+		? InArgs._Icon.Get()
+		: ActionButtonStyle->IconBrush.IsSet()
+		? &ActionButtonStyle->IconBrush.GetValue()
+		: nullptr;
+
+	const bool bHasIcon = Icon.Get() || Icon.IsBound();
+
+	// Empty/default args will resolve from the ActionButtonStyle
+	const TSharedRef<SWidget> ButtonContent =
+		UE::ToolWidgets::Private::ActionButton::MakeButtonContent(
+			ActionButtonStyle,
+			Icon,
+			{},
+			InArgs._Text,
+			{});
+
+	SButton::Construct(
+		SButton::FArguments()
+		.ContentPadding(ActionButtonStyle->GetButtonContentPadding())
+		.ButtonStyle(bHasIcon ? &ActionButtonStyle->GetIconButtonStyle() : &ActionButtonStyle->ButtonStyle)
+		.IsEnabled(InArgs._IsEnabled)
+		.ToolTipText(InArgs._ToolTipText)
+		.HAlign(static_cast<EHorizontalAlignment>(ActionButtonStyle->HorizontalContentAlignment))
+		.VAlign(VAlign_Center)
 		.OnClicked(InArgs._OnClicked)
-	];
+		[
+			ButtonContent
+		]);
 }
diff --git a/Engine/Source/Developer/ToolWidgets/Private/SSimpleButton.cpp b/Engine/Source/Developer/ToolWidgets/Private/SSimpleButton.cpp
--- a/Engine/Source/Developer/ToolWidgets/Private/SSimpleButton.cpp
+++ b/Engine/Source/Developer/ToolWidgets/Private/SSimpleButton.cpp
@@ -2,21 +2,42 @@
 
 #include "SSimpleButton.h"
 
-#include "SActionButton.h"
 #include "ToolWidgetsStyle.h"
-#include "Widgets/SBoxPanel.h"
+#include "ToolWidgetsUtilitiesPrivate.h"
 #include "Widgets/Images/SImage.h"
+#include "Widgets/SBoxPanel.h"
 #include "Widgets/Text/STextBlock.h"
 
 void SSimpleButton::Construct(const FArguments& InArgs)
 {
-	ChildSlot
-	[
-		SNew(SActionButton)
-		.ActionButtonStyle(&UE::ToolWidgets::FToolWidgetsStyle::Get().GetWidgetStyle<FActionButtonStyle>("SimpleButton"))
+	const FActionButtonStyle* ActionButtonStyle = &UE::ToolWidgets::FToolWidgetsStyle::Get().GetWidgetStyle<FActionButtonStyle>("SimpleButton");
+
+	// Check for widget level override, then style override, otherwise unset
+	const TAttribute<const FSlateBrush*> Icon = InArgs._Icon.IsSet()
+		? InArgs._Icon.Get()
+		: ActionButtonStyle->IconBrush.IsSet()
+		? &ActionButtonStyle->IconBrush.GetValue()
+		: nullptr;
+
+	// Empty/default args will resolve from the ActionButtonStyle
+	const TSharedRef<SWidget> ButtonContent =
+		UE::ToolWidgets::Private::ActionButton::MakeButtonContent(
+			ActionButtonStyle,
+			Icon,
+			{},
+			InArgs._Text,
+			{});
+
+	SButton::Construct(
+		SButton::FArguments()
+		.ContentPadding(ActionButtonStyle->GetButtonContentPadding())
 		.ButtonStyle(&FAppStyle::Get().GetWidgetStyle<FButtonStyle>(InArgs._Text.IsSet() ? "SimpleButtonLabelAndIcon" : "SimpleButton"))
-		.Text(InArgs._Text)
-		.Icon(InArgs._Icon)
+		.IsEnabled(InArgs._IsEnabled)
+		.ToolTipText(InArgs._ToolTipText)
+		.HAlign(static_cast<EHorizontalAlignment>(ActionButtonStyle->HorizontalContentAlignment))
+		.VAlign(VAlign_Center)
 		.OnClicked(InArgs._OnClicked)
-	];
+		[
+			ButtonContent
+		]);
 }
diff --git a/Engine/Source/Developer/ToolWidgets/Private/SSimpleComboButton.cpp b/Engine/Source/Developer/ToolWidgets/Private/SSimpleComboButton.cpp
--- a/Engine/Source/Developer/ToolWidgets/Private/SSimpleComboButton.cpp
+++ b/Engine/Source/Developer/ToolWidgets/Private/SSimpleComboButton.cpp
@@ -2,47 +2,68 @@
 
 #include "SSimpleComboButton.h"
 
-#include "SActionButton.h"
-#include "Widgets/SBoxPanel.h"
-#include "Widgets/Images/SImage.h"
-#include "Widgets/Text/STextBlock.h"
 #include "Styling/SlateTypes.h"
 #include "ToolWidgetsStyle.h"
+#include "ToolWidgetsUtilitiesPrivate.h"
+#include "Widgets/Images/SImage.h"
+#include "Widgets/SBoxPanel.h"
+#include "Widgets/Text/STextBlock.h"
 
 void SSimpleComboButton::Construct(const FArguments& InArgs)
 {
-	ActionButton = SNew(SActionButton)
-		.ActionButtonStyle(&UE::ToolWidgets::FToolWidgetsStyle::Get().GetWidgetStyle<FActionButtonStyle>("SimpleComboButton"))
-		.Text(InArgs._Text)
-		.TextBlockStyle(
+	const FActionButtonStyle* ActionButtonStyle = &UE::ToolWidgets::FToolWidgetsStyle::Get().GetWidgetStyle<FActionButtonStyle>("SimpleComboButton");
+
+	// Check for widget level override, then style override, otherwise unset
+	const TAttribute<const FSlateBrush*> Icon = InArgs._Icon.IsSet()
+		? InArgs._Icon.Get()
+		: ActionButtonStyle->IconBrush.IsSet()
+		? &ActionButtonStyle->IconBrush.GetValue()
+		: nullptr;
+
+	const bool bHasIcon = Icon.Get() || Icon.IsBound();
+
+	// Empty/default args will resolve from the ActionButtonStyle
+	const TSharedRef<SWidget> ButtonContent =
+		UE::ToolWidgets::Private::ActionButton::MakeButtonContent(
+			ActionButtonStyle,
+			Icon,
+			{},
+			InArgs._Text,
 			InArgs._UsesSmallText
-			? &FAppStyle::GetWidgetStyle<FTextBlockStyle>("SmallText")
-			: &FAppStyle::GetWidgetStyle<FTextBlockStyle>("SmallButtonText"))
-		.Icon(InArgs._Icon)
+				? &FAppStyle::GetWidgetStyle<FTextBlockStyle>("SmallText")
+				: &FAppStyle::GetWidgetStyle<FTextBlockStyle>("SmallButtonText"));
+
+	const TAttribute<FMargin> ComboButtonContentPadding = ActionButtonStyle->GetComboButtonContentPadding();
+
+	SComboButton::Construct(SComboButton::FArguments()
+		.HasDownArrow(InArgs._HasDownArrow)
+		.ContentPadding(ComboButtonContentPadding)
+		.ButtonStyle(bHasIcon ? &ActionButtonStyle->GetIconButtonStyle() : &ActionButtonStyle->ButtonStyle)
+		.ComboButtonStyle(&ActionButtonStyle->ComboButtonStyle)
+		.IsEnabled(InArgs._IsEnabled)
+		.ToolTipText(InArgs._ToolTipText)
+		.HAlign(static_cast<EHorizontalAlignment>(ActionButtonStyle->HorizontalContentAlignment))
+		.VAlign(VAlign_Center)
+		.ButtonContent()
+		[
+			ButtonContent
+		]
 		.MenuContent()
 		[
 			InArgs._MenuContent.Widget
 		]
 		.OnGetMenuContent(InArgs._OnGetMenuContent)
 		.OnMenuOpenChanged(InArgs._OnMenuOpenChanged)
-		.OnComboBoxOpened(InArgs._OnComboBoxOpened);
-
-	ChildSlot
-	[
-		ActionButton.ToSharedRef()
-	];
+		.OnComboBoxOpened(InArgs._OnComboBoxOpened)
+	);
 }
 
 void SSimpleComboButton::SetMenuContentWidgetToFocus(TWeakPtr<SWidget> InWidget)
 {
-	check(ActionButton.IsValid());
-
-	ActionButton->SetMenuContentWidgetToFocus(InWidget);
+	SComboButton::SetMenuContentWidgetToFocus(InWidget);
 }
 
 void SSimpleComboButton::SetIsMenuOpen(bool bInIsOpen, bool bInIsFocused)
 {
-	check(ActionButton.IsValid());
-
-	ActionButton->SetIsMenuOpen(bInIsOpen, bInIsFocused);
+	SetIsOpen(bInIsOpen, bInIsFocused);
 }
diff --git a/Engine/Source/Developer/ToolWidgets/Private/ToolWidgetsSlateTypes.cpp b/Engine/Source/Developer/ToolWidgets/Private/ToolWidgetsSlateTypes.cpp
--- a/Engine/Source/Developer/ToolWidgets/Private/ToolWidgetsSlateTypes.cpp
+++ b/Engine/Source/Developer/ToolWidgets/Private/ToolWidgetsSlateTypes.cpp
@@ -2,6 +2,8 @@
 
 #include "ToolWidgetsSlateTypes.h"
 
+#include "ToolWidgetsStylePrivate.h"
+
 namespace UE::ToolWidgets::Private
 {
 	inline FName LexToName(const EActionButtonType ActionButtonType)
@@ -103,6 +105,24 @@ FActionButtonStyle& FActionButtonStyle::SetIconButtonStyle(const FButtonStyle& I
 	return *this;
 }
 
+FMargin FActionButtonStyle::GetButtonContentPadding() const
+{
+	static const FMargin DefaultButtonContentPadding = FMargin(
+		UE::ToolWidgets::Private::FToolWidgetsStylePrivate::FActionButton::DefaultHorizontalPadding,
+		UE::ToolWidgets::Private::FToolWidgetsStylePrivate::FActionButton::DefaultVerticalPadding);
+
+	return ButtonContentPadding.IsSet()
+		? ButtonContentPadding.GetValue()
+		: DefaultButtonContentPadding;
+}
+
+FMargin FActionButtonStyle::GetComboButtonContentPadding() const
+{
+	return ComboButtonContentPadding.IsSet()
+		? ComboButtonContentPadding.GetValue()
+		: ComboButtonStyle.ContentPadding;
+}
+
 FActionButtonStyle& FActionButtonStyle::SetButtonContentPadding(const FMargin& InContentPadding)
 {
 	ButtonContentPadding = InContentPadding;
diff --git a/Engine/Source/Developer/ToolWidgets/Private/ToolWidgetsStyle.cpp b/Engine/Source/Developer/ToolWidgets/Private/ToolWidgetsStyle.cpp
--- a/Engine/Source/Developer/ToolWidgets/Private/ToolWidgetsStyle.cpp
+++ b/Engine/Source/Developer/ToolWidgets/Private/ToolWidgetsStyle.cpp
@@ -107,7 +107,7 @@ namespace UE::ToolWidgets
 					.SetButtonStyle(ComboButtonStyle.ButtonStyle)
 					.SetComboButtonContentPadding(ModifyContentPadding(ComboButtonStyle, DefaultButtonContentPadding))
 					.SetHorizontalContentAlignment(HAlign_Left)
-					.SetTextBlockStyle(GetWidgetStyle<FTextBlockStyle>("SmallText"))
+					.SetTextBlockStyle(GetWidgetStyle<FTextBlockStyle>("SmallButtonText"))
 					.SetIconColorAndOpacity(FSlateColor::UseForeground());
 
 				Set("SimpleComboButton", SimpleComboButton);
diff --git a/Engine/Source/Developer/ToolWidgets/Private/ToolWidgetsUtilitiesPrivate.cpp b/Engine/Source/Developer/ToolWidgets/Private/ToolWidgetsUtilitiesPrivate.cpp
new file mode 100644
--- /dev/null
+++ b/Engine/Source/Developer/ToolWidgets/Private/ToolWidgetsUtilitiesPrivate.cpp
@@ -0,0 +1,108 @@
+// Copyright Epic Games, Inc. All Rights Reserved.
+
+#include "ToolWidgetsUtilitiesPrivate.h"
+
+#include "Styling/StyleColors.h"
+#include "ToolWidgetsSlateTypes.h"
+#include "ToolWidgetsStylePrivate.h"
+#include "Widgets/Images/SImage.h"
+#include "Widgets/Input/SButton.h"
+#include "Widgets/Layout/SBox.h"
+#include "Widgets/Layout/SSpacer.h"
+#include "Widgets/Layout/SWidgetSwitcher.h"
+#include "Widgets/SBoxPanel.h"
+#include "Widgets/Text/STextBlock.h"
+
+namespace UE::ToolWidgets::Private
+{
+	namespace ActionButton
+	{
+		TSharedRef<SWidget> MakeButtonContent(
+			const TAttribute<const FSlateBrush*>& InIcon,
+			const TAttribute<FSlateColor>& InIconColorAndOpacity,
+			const TAttribute<FText>& InText,
+			const FTextBlockStyle* InTextBlockStyle)
+		{
+			check(InIconColorAndOpacity.IsSet() || InIconColorAndOpacity.IsBound());
+			check(InTextBlockStyle);
+
+			const bool bHasIcon = InIcon.Get() || InIcon.IsBound();
+
+			static constexpr float DefaultIconHeight = FToolWidgetsStylePrivate::FActionButton::DefaultIconHeight;
+			static constexpr float IconTextPadding = FToolWidgetsStylePrivate::FActionButton::DefaultIconLabelSpacing;
+
+			const TSharedRef<SHorizontalBox> ButtonContentContainer = SNew(SHorizontalBox)
+				+ SHorizontalBox::Slot()
+				.HAlign(HAlign_Center)
+				.VAlign(VAlign_Center)
+				.AutoWidth()
+				.Padding(0)
+				[
+					SNew(SWidgetSwitcher)
+					.WidgetIndex(bHasIcon ? 1 : 0)
+
+					+ SWidgetSwitcher::Slot()
+					[
+						SNew(SSpacer)
+						.Size(FVector2D{ 0, DefaultIconHeight })
+					]
+
+					+ SWidgetSwitcher::Slot()
+					[
+						SNew(SImage)
+						.Image(InIcon)
+						.ColorAndOpacity(InIconColorAndOpacity)
+						.Visibility(bHasIcon ? EVisibility::HitTestInvisible : EVisibility::Collapsed)
+					]
+				]
+
+				+ SHorizontalBox::Slot()
+				.VAlign(VAlign_Center)
+				.FillWidth(1.0f)
+				.Padding(FMargin(bHasIcon ? IconTextPadding : 0, 0, 0, 0))
+				[
+					SNew(STextBlock)
+					.TextStyle(InTextBlockStyle)
+					.Text(InText)
+					.Visibility_Lambda([InText]()
+					{
+						return InText.Get(
+							FText::GetEmpty()).IsEmpty()
+							? EVisibility::Collapsed
+							: EVisibility::Visible;
+					})
+				];
+
+			return ButtonContentContainer;
+		}
+
+		TSharedRef<SWidget> MakeButtonContent(
+			const FActionButtonStyle* InActionButtonStyle,
+			const TAttribute<const FSlateBrush*>& InIcon,
+			const TAttribute<FSlateColor>& InIconColorAndOpacity,
+			const TAttribute<FText>& InText,
+			const FTextBlockStyle* InTextBlockStyle)
+		{
+			check(InActionButtonStyle);
+
+			// Check for widget level override, then style override, otherwise unset
+			const TAttribute<const FSlateBrush*> Icon =
+				InIcon.IsSet()
+				? InIcon.Get()
+				: InActionButtonStyle->IconBrush.IsSet()
+				? &InActionButtonStyle->IconBrush.GetValue()
+				: nullptr;
+
+			TAttribute<FSlateColor> IconColorAndOpacity = FSlateColor::UseForeground();
+			if (ensureMsgf(InIconColorAndOpacity.IsSet() || InActionButtonStyle->IconColorAndOpacity.IsSet(), TEXT("The provided IconColorAndOpacity must either be set directly, or stored in the ActionButtonStyle.")))
+			{
+				// If provided Attribute not set, get from the ActionButtonStyle
+				IconColorAndOpacity = InIconColorAndOpacity.Get(InActionButtonStyle->IconColorAndOpacity.Get(IconColorAndOpacity.Get()));
+			}
+
+			const FTextBlockStyle* TextBlockStyle = InTextBlockStyle ? InTextBlockStyle : &InActionButtonStyle->TextBlockStyle;
+
+			return MakeButtonContent(Icon, IconColorAndOpacity, InText, TextBlockStyle);
+		}
+	}
+}
diff --git a/Engine/Source/Developer/ToolWidgets/Private/ToolWidgetsUtilitiesPrivate.h b/Engine/Source/Developer/ToolWidgets/Private/ToolWidgetsUtilitiesPrivate.h
new file mode 100644
--- /dev/null
+++ b/Engine/Source/Developer/ToolWidgets/Private/ToolWidgetsUtilitiesPrivate.h
@@ -0,0 +1,27 @@
+// Copyright Epic Games, Inc. All Rights Reserved.
+
+#pragma once
+
+#include "ToolWidgetsSlateTypes.h"
+#include "Widgets/Images/SImage.h"
+#include "Widgets/Input/SButton.h"
+#include "Widgets/Layout/SSpacer.h"
+
+namespace UE::ToolWidgets::Private
+{
+	namespace ActionButton
+	{
+		TSharedRef<SWidget> MakeButtonContent(
+			const TAttribute<const FSlateBrush*>& InIcon,
+			const TAttribute<FSlateColor>& InIconColorAndOpacity,
+			const TAttribute<FText>& InText,
+			const FTextBlockStyle* InTextBlockStyle);
+
+		TSharedRef<SWidget> MakeButtonContent(
+			const FActionButtonStyle* InActionButtonStyle,
+			const TAttribute<const FSlateBrush*>& InIcon,
+			const TAttribute<FSlateColor>& InIconColorAndOpacity,
+			const TAttribute<FText>& InText,
+			const FTextBlockStyle* InTextBlockStyle);
+	}
+}
diff --git a/Engine/Source/Developer/ToolWidgets/Public/SPrimaryButton.h b/Engine/Source/Developer/ToolWidgets/Public/SPrimaryButton.h
--- a/Engine/Source/Developer/ToolWidgets/Public/SPrimaryButton.h
+++ b/Engine/Source/Developer/ToolWidgets/Public/SPrimaryButton.h
@@ -6,11 +6,8 @@
 #include "Widgets/DeclarativeSyntaxSupport.h"
 #include "Widgets/Input/SButton.h"
 
-class SActionButton;
-class SHorizontalBox;
-
 class TOOLWIDGETS_API SPrimaryButton
-	: public SCompoundWidget
+	: public SButton
 {
 public:
 	SLATE_BEGIN_ARGS(SPrimaryButton)
@@ -28,7 +25,4 @@ public:
 	SLATE_END_ARGS()
 
 	void Construct(const FArguments& InArgs);
-
-private:
-	TSharedPtr<SActionButton> ActionButton;
 };
diff --git a/Engine/Source/Developer/ToolWidgets/Public/SSimpleButton.h b/Engine/Source/Developer/ToolWidgets/Public/SSimpleButton.h
--- a/Engine/Source/Developer/ToolWidgets/Public/SSimpleButton.h
+++ b/Engine/Source/Developer/ToolWidgets/Public/SSimpleButton.h
@@ -8,7 +8,7 @@
 #include "Widgets/SCompoundWidget.h"
 
 class TOOLWIDGETS_API SSimpleButton
-	: public SCompoundWidget
+	: public SButton
 {
 public:
 	SLATE_BEGIN_ARGS(SSimpleButton)
diff --git a/Engine/Source/Developer/ToolWidgets/Public/SSimpleComboButton.h b/Engine/Source/Developer/ToolWidgets/Public/SSimpleComboButton.h
--- a/Engine/Source/Developer/ToolWidgets/Public/SSimpleComboButton.h
+++ b/Engine/Source/Developer/ToolWidgets/Public/SSimpleComboButton.h
@@ -5,11 +5,11 @@
 #include "Framework/SlateDelegates.h"
 #include "SActionButton.h"
 #include "Widgets/DeclarativeSyntaxSupport.h"
-#include "Widgets/SCompoundWidget.h"
 #include "Widgets/Input/SComboButton.h"
+#include "Widgets/SCompoundWidget.h"
 
 class TOOLWIDGETS_API SSimpleComboButton
-	: public SCompoundWidget
+	: public SComboButton
 	, public IActionButton
 {
 public:
@@ -40,7 +40,4 @@ public:
 	virtual void SetMenuContentWidgetToFocus(TWeakPtr<SWidget> InWidget) override;
 	virtual void SetIsMenuOpen(bool bInIsOpen, bool bInIsFocused) override;
 	//~ End IActionButton
-
-private:
-	TSharedPtr<SActionButton> ActionButton;
 };
diff --git a/Engine/Source/Developer/ToolWidgets/Public/ToolWidgetsSlateTypes.h b/Engine/Source/Developer/ToolWidgets/Public/ToolWidgetsSlateTypes.h
--- a/Engine/Source/Developer/ToolWidgets/Public/ToolWidgetsSlateTypes.h
+++ b/Engine/Source/Developer/ToolWidgets/Public/ToolWidgetsSlateTypes.h
@@ -61,6 +61,8 @@ struct FActionButtonStyle : public FSlateWidgetStyle
 	 */
 	UPROPERTY(EditAnywhere, BlueprintReadWrite, Category = "Appearance")
 	TOptional<FMargin> ButtonContentPadding;
+
+	FMargin GetButtonContentPadding() const;
 	FActionButtonStyle& SetButtonContentPadding(const FMargin& InContentPadding);
 
 	/**
@@ -82,6 +84,8 @@ struct FActionButtonStyle : public FSlateWidgetStyle
 	 */
 	UPROPERTY(EditAnywhere, BlueprintReadWrite, Category = "Appearance")
 	TOptional<FMargin> ComboButtonContentPadding;
+
+	FMargin GetComboButtonContentPadding() const;
 	FActionButtonStyle& SetComboButtonContentPadding(const FMargin& InContentPadding);
 
 	/**
