[Sequencer] Fix Sequencer sidebar drawers causing viewport focus issues.

diff --git a/Engine/Source/Developer/ToolWidgets/Private/Sidebar/SSidebar.cpp b/Engine/Source/Developer/ToolWidgets/Private/Sidebar/SSidebar.cpp
--- a/Engine/Source/Developer/ToolWidgets/Private/Sidebar/SSidebar.cpp
+++ b/Engine/Source/Developer/ToolWidgets/Private/Sidebar/SSidebar.cpp
@@ -277,19 +277,19 @@ void SSidebar::OnDrawerTabDockToggled(const TSharedRef<FSidebarDrawer>& InDrawer
 
 void SSidebar::OnTabDrawerFocusLost(const TSharedRef<SSidebarDrawer>& InDrawerWidget)
 {
-	const TSharedPtr<FSidebarDrawer> DrawerWidget = InDrawerWidget->GetDrawer();
-	if (!DrawerWidget.IsValid())
+	const TSharedPtr<FSidebarDrawer> Drawer = InDrawerWidget->GetDrawer();
+	if (!Drawer.IsValid())
 	{
 		return;
 	}
 
 	// Don't automatically close a pinned tab that is in the foreground
-	if (IsDrawerPinned(DrawerWidget->GetUniqueId()) && DrawerWidget == GetForegroundTab())
+	if (IsDrawerPinned(Drawer->GetUniqueId()) && Drawer == GetForegroundTab())
 	{
 		return;
 	}
 
-	CloseDrawerInternal(DrawerWidget.ToSharedRef());
+	CloseDrawerInternal(Drawer.ToSharedRef());
 }
 
 void SSidebar::OnTabDrawerClosed(const TSharedRef<SSidebarDrawer>& InDrawerWidget)
@@ -535,18 +535,22 @@ void SSidebar::CloseDrawerInternal(const TSharedRef<FSidebarDrawer>& InDrawer)
 {
 	if (const TSharedPtr<SSidebarDrawer> OpenedDrawer = FindOpenedDrawer(InDrawer))
 	{
-		OpenedDrawer->Close();
 		const TSharedRef<SSidebarDrawer> OpenedDrawerRef = OpenedDrawer.ToSharedRef();
 
+		// Setting bInAnimateOpen to true causes a bug that forces keyboard focus to the drawer.
+		// Setting this false until this can be worked out. For some reason, inside Close(), the
+		// RegisterActiveTimer doesn't actually start a timer. This works for Open() though.
+		OpenedDrawerRef->Close(/*bInAnimateOpen=*/false);
+
 		if (const TSharedPtr<SOverlay> DrawersOverlay = DrawersOverlayWeak.Pin())
 		{
 			DrawersOverlay->RemoveSlot(OpenedDrawerRef);
 		}
 
 		OpenedDrawers.Remove(OpenedDrawerRef);
-	}
 
-	InDrawer->bIsOpen = false;
+		InDrawer->bIsOpen = false;
+	}
 
 	SummonPinnedTabIfNothingOpened();
 	UpdateDrawerAppearance();
diff --git a/Engine/Source/Developer/ToolWidgets/Private/Sidebar/SSidebarDrawer.cpp b/Engine/Source/Developer/ToolWidgets/Private/Sidebar/SSidebarDrawer.cpp
--- a/Engine/Source/Developer/ToolWidgets/Private/Sidebar/SSidebarDrawer.cpp
+++ b/Engine/Source/Developer/ToolWidgets/Private/Sidebar/SSidebarDrawer.cpp
@@ -33,7 +33,7 @@ void SSidebarDrawer::Construct(const FArguments& InArgs, const TSharedRef<FSideb
 	OnDrawerFocusLost = InArgs._OnDrawerFocusLost;
 	OnDrawerClosed = InArgs._OnDrawerClosed;
 
-	OpenCloseAnimation = FCurveSequence(0.f, 0.15f, ECurveEaseFunction::QuadOut);
+	OpenCloseAnimation = FCurveSequence(0.f, AnimationLength, ECurveEaseFunction::QuadOut);
 
 	FSlateApplication::Get().OnFocusChanging().AddSP(this, &SSidebarDrawer::OnGlobalFocusChanging);
 
@@ -52,35 +52,41 @@ void SSidebarDrawer::SetCurrentSize(const float InSize)
 	CurrentSize = FMath::Clamp(InSize, MinDrawerSize, TargetDrawerSize);
 }
 
-void SSidebarDrawer::Open(bool bAnimateOpen)
+void SSidebarDrawer::Open(const bool bInAnimateOpen)
 {
-	if (!bAnimateOpen)
+	if (!bInAnimateOpen)
 	{
 		SetCurrentSize(TargetDrawerSize);
 		OpenCloseAnimation.JumpToEnd();
 		return;
 	}
 
-	OpenCloseAnimation.Play(AsShared(), false, OpenCloseAnimation.IsPlaying() ? OpenCloseAnimation.GetSequenceTime() : 0.0f, false);
+	const float StartTime = OpenCloseAnimation.IsPlaying() ? OpenCloseAnimation.GetSequenceTime() : 0.f;
+	OpenCloseAnimation.Play(AsShared(), false, StartTime, false);
 
 	if (!OpenCloseTimer.IsValid())
 	{
 		AnimationThrottle = FSlateThrottleManager::Get().EnterResponsiveMode();
-		OpenCloseTimer = RegisterActiveTimer(0.0f, FWidgetActiveTimerDelegate::CreateSP(this, &SSidebarDrawer::UpdateAnimation));
+		OpenCloseTimer = RegisterActiveTimer(0.f, FWidgetActiveTimerDelegate::CreateSP(this, &SSidebarDrawer::UpdateAnimation));
 	}
 }
 
-void SSidebarDrawer::Close()
+void SSidebarDrawer::Close(const bool bInAnimateOpen)
 {
-	if (OpenCloseAnimation.IsForward())
+	if (!bInAnimateOpen)
 	{
-		OpenCloseAnimation.Reverse();
+		SetCurrentSize(0.f);
+		OpenCloseAnimation.JumpToStart();
+		return;
 	}
 
+	const float StartTime = OpenCloseAnimation.IsPlaying() ? OpenCloseAnimation.GetSequenceTime() : AnimationLength;
+	OpenCloseAnimation.PlayReverse(AsShared(), false, StartTime, true);
+
 	if (!OpenCloseTimer.IsValid())
 	{
 		AnimationThrottle = FSlateThrottleManager::Get().EnterResponsiveMode();
-		OpenCloseTimer = RegisterActiveTimer(0.0f, FWidgetActiveTimerDelegate::CreateSP(this, &SSidebarDrawer::UpdateAnimation));
+		OpenCloseTimer = RegisterActiveTimer(0.f, FWidgetActiveTimerDelegate::CreateSP(this, &SSidebarDrawer::UpdateAnimation));
 	}
 }
 
@@ -199,7 +205,7 @@ FReply SSidebarDrawer::OnMouseMove(const FGeometry& InAllottedGeometry, const FP
 		const FVector2D MousePosition = InMouseEvent.GetScreenSpacePosition();
 		const FVector2D LocalMousePosition = InitialResizeGeometry.AbsoluteToLocal(MousePosition);
 
-		float DeltaSize = 0.0f;
+		float DeltaSize = 0.f;
 
 		switch (TabLocation)
 		{
@@ -522,7 +528,6 @@ void SSidebarDrawer::OnGlobalFocusChanging(const FFocusEvent& InFocusEvent
 {
 	// Sometimes when dismissing focus can change which will trigger this again
 	static bool bIsReEntrant = false;
-
 	if (bIsReEntrant)
 	{
 		return;
@@ -560,11 +565,10 @@ void SSidebarDrawer::OnGlobalFocusChanging(const FFocusEvent& InFocusEvent
 	{
 		// New focus is on something else, try to check if it's a menu or child window
 		const TSharedRef<SWindow> NewWindow = InNewFocusedWidgetPath.GetWindow();
-		const TSharedPtr<SWindow> MyWindow = FSlateApplication::Get().FindWidgetWindow(ThisWidget);
+		const TSharedPtr<SWindow> ThisWindow = FSlateApplication::Get().FindWidgetWindow(ThisWidget);
 
 		// See if this is a child window (like a color picker being opened from details), and if so, don't dismiss
-		// Rely on OnActiveTabChanged below to lose focus if the child window actually contains tabs
-		if (!NewWindow->IsDescendantOf(MyWindow))
+		if (!NewWindow->IsDescendantOf(ThisWindow))
 		{
 			if (const TSharedPtr<SWidget> MenuHost = FSlateApplication::Get().GetMenuHostWidget())
 			{
@@ -594,12 +598,3 @@ void SSidebarDrawer::OnGlobalFocusChanging(const FFocusEvent& InFocusEvent
 		OnDrawerFocusLost.ExecuteIfBound(ThisWidget);
 	}
 }
-
-void SSidebarDrawer::OnActiveTabChanged(const TSharedPtr<FSidebarDrawer>& InPreviouslyActive, const TSharedPtr<FSidebarDrawer>& InNewlyActivated)
-{
-	// This tab lost the active status to some other tab; treat this like focus was lost
-	if (InPreviouslyActive == DrawerWeak && InNewlyActivated.IsValid())
-	{
-		OnDrawerFocusLost.ExecuteIfBound(SharedThis(this));
-	}
-}
diff --git a/Engine/Source/Developer/ToolWidgets/Private/Sidebar/SSidebarDrawer.h b/Engine/Source/Developer/ToolWidgets/Private/Sidebar/SSidebarDrawer.h
--- a/Engine/Source/Developer/ToolWidgets/Private/Sidebar/SSidebarDrawer.h
+++ b/Engine/Source/Developer/ToolWidgets/Private/Sidebar/SSidebarDrawer.h
@@ -58,10 +58,14 @@ public:
 	 *
 	 * @param bInAnimateOpen Whether to play an animation when opening the drawer, defaults to true.
 	 */
-	void Open(bool bInAnimateOpen = true);
+	void Open(const bool bInAnimateOpen = true);
 
-	/** Begins an animation which closes the drawer. */
-	void Close();
+	/**
+	 * Closes the drawer.
+	 *
+	 * @param bInAnimateOpen Whether to play an animation when closing the drawer, defaults to true.
+	 */
+	void Close(const bool bInAnimateOpen = true);
 
 	/** @return True if the drawer is open. */
 	bool IsOpen() const;
@@ -75,6 +79,8 @@ public:
 	TSharedPtr<FSidebarDrawer> GetDrawer() const;
 
 protected:
+	static constexpr float AnimationLength = 0.15f;
+
 	//~ Begin SWidget
 	virtual bool SupportsKeyboardFocus() const override;
 	virtual FVector2D ComputeDesiredSize(const float InLayoutScaleMultiplier) const override;
@@ -89,20 +95,18 @@ protected:
 	//~ End SWidget
 
 	FGeometry GetRenderTransformedGeometry(const FGeometry& InAllottedGeometry) const;
-	
+
 	FGeometry GetResizeHandleGeometry(const FGeometry& InAllottedGeometry) const;
-	
+
 	EActiveTimerReturnType UpdateAnimation(const double InCurrentTime, const float InDeltaTime);
-	
+
 	void OnGlobalFocusChanging(const FFocusEvent& InFocusEvent
 		, const FWeakWidgetPath& InOldFocusedWidgetPath, const TSharedPtr<SWidget>& InOldFocusedWidget
 		, const FWidgetPath& InNewFocusedWidgetPath, const TSharedPtr<SWidget>& InNewFocusedWidget);
-	
-	void OnActiveTabChanged(const TSharedPtr<FSidebarDrawer>& InNewlyActivated, const TSharedPtr<FSidebarDrawer>& InPreviouslyActive);
 
 	TWeakPtr<FSidebarDrawer> DrawerWeak;
 	ESidebarTabLocation TabLocation = ESidebarTabLocation::Right;
-	
+
 	float MinDrawerSize = 0.f;
 	float MaxDrawerSize = 0.f;
 	float TargetDrawerSize = 0.f;
@@ -119,9 +123,9 @@ protected:
 
 	FThrottleRequest ResizeThrottleHandle;
 	FThrottleRequest AnimationThrottle;
-	
+
 	float CurrentSize = 0.f;
-	
+
 	bool bIsResizing = false;
 	bool bIsResizeHandleHovered = false;
 	float InitialSizeAtResize = 0.f;
