[ToolWidgets]
 - Fix SSidebar close animation not playing.
 - Rename a few variables/properties for clarity
 - Show focused style for tab button when docked

diff --git a/Engine/Source/Developer/ToolWidgets/Private/Sidebar/SSidebar.cpp b/Engine/Source/Developer/ToolWidgets/Private/Sidebar/SSidebar.cpp
--- a/Engine/Source/Developer/ToolWidgets/Private/Sidebar/SSidebar.cpp
+++ b/Engine/Source/Developer/ToolWidgets/Private/Sidebar/SSidebar.cpp
@@ -46,7 +46,7 @@ void SSidebar::Construct(const FArguments& InArgs, const TSharedRef<SOverlay>& I
 		.Padding(0.0f)
 		.BorderImage(FAppStyle::Get().GetBrush(TEXT("Docking.Sidebar.Background")))
 		[
-			SAssignNew(TabContainer, SScrollBox)
+			SAssignNew(TabButtonContainer, SScrollBox)
 			.Orientation(IsHorizontal() ? EOrientation::Orient_Horizontal : EOrientation::Orient_Vertical)
 			.ScrollBarAlwaysVisible(false)
 			.ScrollBarVisibility(EVisibility::Collapsed)
@@ -69,7 +69,7 @@ bool SSidebar::RegisterDrawer(FSidebarDrawerConfig&& InDrawerConfig)
 		: NewDrawer->ContentWidget = SNew(SSidebarDrawerContent, NewDrawer);
 
 	// Add tab button
-	TabContainer->AddSlot()
+	TabButtonContainer->AddSlot()
 		[
 			SAssignNew(NewDrawer->ButtonWidget, SSidebarButton, NewDrawer, TabLocation)
 			.MinButtonSize(bAlwaysUseMaxButtonSize ? MaxTabButtonSize : MinTabButtonSize)
@@ -81,7 +81,7 @@ bool SSidebar::RegisterDrawer(FSidebarDrawerConfig&& InDrawerConfig)
 			.OnGetContextMenuContent(this, &SSidebar::OnGetTabDrawerContextMenuWidget, NewDrawer)
 		];
 
-	DrawerTabs.Add(NewDrawer);
+	Drawers.Add(NewDrawer);
 
 	// Figure out the size this tab should be when opened later. We do it now when the tab still has valid geometry. Once it is moved to the sidebar it will not.
 	float TargetDrawerSizePct = NewDrawer->SizeCoefficient;
@@ -111,7 +111,9 @@ bool SSidebar::RegisterDrawer(FSidebarDrawerConfig&& InDrawerConfig)
 	{
 		SetDrawerDocked(NewDrawer->GetUniqueId(), true);
 	}
-	
+
+	UpdateDrawerAppearance();
+
 	return true;
 }
 
@@ -122,20 +124,21 @@ bool SSidebar::UnregisterDrawer(const FName InDrawerId)
 		CloseAllDrawers();
 	}
 	
-	const int32 FoundIndex = DrawerTabs.IndexOfByPredicate(
+	const int32 IndexToRemove = Drawers.IndexOfByPredicate(
 		[InDrawerId](const TSharedRef<FSidebarDrawer>& InDrawer)
 		{
 			return InDrawerId == InDrawer->GetUniqueId();
 		});
-	if (FoundIndex == INDEX_NONE)
+	if (IndexToRemove == INDEX_NONE)
 	{
 		return false;
 	}
 
-	DrawerTabs.RemoveAt(FoundIndex);
-	TabContainer->RemoveSlot(DrawerTabs[FoundIndex]->ButtonWidget.ToSharedRef());
+	TabButtonContainer->RemoveSlot(Drawers[IndexToRemove]->ButtonWidget.ToSharedRef());
+
+	RemoveDrawer(Drawers[IndexToRemove]);
+	Drawers.RemoveAt(IndexToRemove);
 
-	RemoveDrawer(DrawerTabs[FoundIndex]);
 	SummonPinnedTabIfNothingOpened();
 
 	// Clear the pinned flag when the tab is removed from the sidebar.
@@ -143,10 +146,14 @@ bool SSidebar::UnregisterDrawer(const FName InDrawerId)
 	// then moving it to the sidebar again will leave it unpinned the second time.)
 	SetDrawerPinned(InDrawerId, false);
 
-	if (DrawerTabs.Num() == 0)
+	if (Drawers.Num() == 0)
 	{
 		SetVisibility(EVisibility::Collapsed);
 	}
+	else
+	{
+		UpdateDrawerAppearance();
+	}
 
 	return true;
 }
@@ -158,7 +165,7 @@ bool SSidebar::ContainsDrawer(const FName InDrawerId) const
 
 int32 SSidebar::GetDrawerCount() const
 {
-	return DrawerTabs.Num();
+	return Drawers.Num();
 }
 
 bool SSidebar::RegisterDrawerSection(const FName InDrawerId, const TSharedPtr<ISidebarDrawerContent>& InSection)
@@ -210,21 +217,22 @@ bool SSidebar::TryOpenDrawer(const FName InDrawerId)
 	{
 		return false;
 	}
-	
-	if (const TSharedPtr<FSidebarDrawer> InDrawer = FindDrawer(InDrawerId))
+
+	const TSharedPtr<FSidebarDrawer> Drawer = FindDrawer(InDrawerId);
+	if (!Drawer.IsValid())
 	{
-		OpenDrawerNextFrame(InDrawer.ToSharedRef(), /*bAnimateOpen=*/ true);
-		return true;
+		return false;
 	}
-	
-	return false;
+
+	OpenDrawerNextFrame(Drawer.ToSharedRef());
+	return true;
 }
 
-void SSidebar::CloseAllDrawers()
+void SSidebar::CloseAllDrawers(const bool bInAnimate)
 {
-	for (const TSharedRef<FSidebarDrawer>& DrawerTab : DrawerTabs)
+	for (const TSharedRef<FSidebarDrawer>& DrawerTab : Drawers)
 	{
-		CloseDrawerInternal(DrawerTab);
+		CloseDrawerInternal(DrawerTab, bInAnimate);
 	}
 }
 
@@ -241,14 +249,18 @@ void SSidebar::OnTabDrawerButtonPressed(const TSharedRef<FSidebarDrawer>& InDraw
 	else if (!InDrawer->bIsDocked)
 	{
 		// Otherwise clicking on an inactive tab should open the drawer
-		OpenDrawerInternal(InDrawer, /*bAnimateOpen=*/true);
+		OpenDrawerInternal(InDrawer);
+	}
+	else if (InDrawer->bIsDocked && InDrawer->DrawerWidget.IsValid())
+	{
+		FSlateApplication::Get().SetKeyboardFocus(InDrawer->DrawerWidget);
 	}
 }
 
 void SSidebar::OnDrawerTabPinToggled(const TSharedRef<FSidebarDrawer>& InDrawer, const bool bIsPinned)
 {
 	// Set pin state for given tab; clear the pin state for all other tabs
-	for (const TSharedRef<FSidebarDrawer>& DrawerTab : DrawerTabs)
+	for (const TSharedRef<FSidebarDrawer>& DrawerTab : Drawers)
 	{
 		SetDrawerPinned(DrawerTab->GetUniqueId(), DrawerTab == InDrawer ? bIsPinned : false);
 	}
@@ -256,7 +268,7 @@ void SSidebar::OnDrawerTabPinToggled(const TSharedRef<FSidebarDrawer>& InDrawer,
 	// Open any newly-pinned tab
 	if (bIsPinned)
 	{
-		OpenDrawerInternal(InDrawer, /*bAnimateOpen=*/true);
+		OpenDrawerInternal(InDrawer);
 	}
 }
 
@@ -278,13 +290,7 @@ void SSidebar::OnDrawerTabDockToggled(const TSharedRef<FSidebarDrawer>& InDrawer
 void SSidebar::OnTabDrawerFocusLost(const TSharedRef<SSidebarDrawer>& InDrawerWidget)
 {
 	const TSharedPtr<FSidebarDrawer> Drawer = InDrawerWidget->GetDrawer();
-	if (!Drawer.IsValid())
-	{
-		return;
-	}
-
-	// Don't automatically close a pinned tab that is in the foreground
-	if (IsDrawerPinned(Drawer->GetUniqueId()) && Drawer == GetForegroundTab())
+	if (!Drawer.IsValid() || IsDrawerPinned(Drawer->GetUniqueId()))
 	{
 		return;
 	}
@@ -292,9 +298,18 @@ void SSidebar::OnTabDrawerFocusLost(const TSharedRef<SSidebarDrawer>& InDrawerWi
 	CloseDrawerInternal(Drawer.ToSharedRef());
 }
 
-void SSidebar::OnTabDrawerClosed(const TSharedRef<SSidebarDrawer>& InDrawerWidget)
+void SSidebar::OnOpenAnimationFinish(const TSharedRef<SSidebarDrawer>& InDrawerWidget)
 {
-	RemoveDrawer(InDrawerWidget->GetDrawer().ToSharedRef());
+}
+
+void SSidebar::OnCloseAnimationFinish(const TSharedRef<SSidebarDrawer>& InDrawerWidget)
+{
+	if (const TSharedPtr<SOverlay> DrawersOverlay = DrawersOverlayWeak.Pin())
+	{
+		DrawersOverlay->RemoveSlot(InDrawerWidget);
+	}
+
+	ClosingDrawerWidgets.Remove(InDrawerWidget);
 }
 
 void SSidebar::OnDrawerTargetSizeChanged(const TSharedRef<SSidebarDrawer>& InDrawerWidget, const float InNewSize)
@@ -418,15 +433,17 @@ void SSidebar::RemoveAllDrawers()
 	bAnimatePendingTabOpen = false;
 
 	// Closing drawers can remove them from the opened drawers list so copy the list first
-	TArray<TSharedRef<SSidebarDrawer>> OpenedDrawersCopy = OpenedDrawers;
+	TArray<TSharedRef<SSidebarDrawer>> OpenDrawerWidgetsCopy = OpenDrawerWidgets;
 
-	for (TSharedRef<SSidebarDrawer>& Drawer : OpenedDrawersCopy)
+	for (const TSharedRef<SSidebarDrawer>& Drawer : OpenDrawerWidgetsCopy)
 	{
 		RemoveDrawer(Drawer->GetDrawer().ToSharedRef());
 	}
+
+	Drawers.Empty();
 }
 
-EActiveTimerReturnType SSidebar::OnOpenPendingDrawerTimer(const double CurrentTime, const float DeltaTime)
+EActiveTimerReturnType SSidebar::OnOpenPendingDrawerTimer(const double InCurrentTime, const float InDeltaTime)
 {
 	if (const TSharedPtr<FSidebarDrawer> TabToOpen = PendingTabToOpen.Pin())
 	{
@@ -450,10 +467,10 @@ EActiveTimerReturnType SSidebar::OnOpenPendingDrawerTimer(const double CurrentTi
 	return EActiveTimerReturnType::Stop;
 }
 
-void SSidebar::OpenDrawerNextFrame(const TSharedRef<FSidebarDrawer>& InDrawer, const bool bAnimateOpen)
+void SSidebar::OpenDrawerNextFrame(const TSharedRef<FSidebarDrawer>& InDrawer, const bool bInAnimate)
 {
 	PendingTabToOpen = InDrawer;
-	bAnimatePendingTabOpen = bAnimateOpen;
+	bAnimatePendingTabOpen = bInAnimate;
 
 	if (!OpenPendingDrawerTimerHandle.IsValid())
 	{
@@ -461,9 +478,9 @@ void SSidebar::OpenDrawerNextFrame(const TSharedRef<FSidebarDrawer>& InDrawer, c
 	}
 }
 
-void SSidebar::OpenDrawerInternal(const TSharedRef<FSidebarDrawer>& InDrawer, const bool bAnimateOpen)
+void SSidebar::OpenDrawerInternal(const TSharedRef<FSidebarDrawer>& InDrawer, const bool bInAnimate)
 {
-	if (FindOpenedDrawer(InDrawer))
+	if (OpenDrawerWidgets.Contains(InDrawer->DrawerWidget))
 	{
 		return;
 	}
@@ -496,27 +513,40 @@ void SSidebar::OpenDrawerInternal(const TSharedRef<FSidebarDrawer>& InDrawer, co
 
 	const float TargetDrawerSize = AvailableWidth * TargetDrawerSizePct;
 
-	InDrawer->DrawerWidget =
-		SNew(SSidebarDrawer, InDrawer, TabLocation)
-		.MinDrawerSize(MinDrawerSize)
-		.TargetDrawerSize(TargetDrawerSize)
-		.MaxDrawerSize(MaxDrawerSize)
-		.OnDrawerFocusLost(this, &SSidebar::OnTabDrawerFocusLost)
-		.OnDrawerClosed(this, &SSidebar::OnTabDrawerClosed)
-		.OnDrawerTargetSizeChanged(this, &SSidebar::OnDrawerTargetSizeChanged);
-
-	DrawersOverlay->AddSlot()
-		.Padding(SlotPadding)
-		.HAlign(TabLocation == ESidebarTabLocation::Left ? HAlign_Left : HAlign_Right)
-		[
-			InDrawer->DrawerWidget.ToSharedRef()
-		];
+	if (!InDrawer->DrawerWidget.IsValid())
+	{
+		InDrawer->DrawerWidget =
+			SNew(SSidebarDrawer, InDrawer, TabLocation)
+			.MinDrawerSize(MinDrawerSize)
+			.TargetDrawerSize(TargetDrawerSize)
+			.MaxDrawerSize(MaxDrawerSize)
+			.OnDrawerFocusLost(this, &SSidebar::OnTabDrawerFocusLost)
+			.OnOpenAnimationFinish(this, &SSidebar::OnOpenAnimationFinish)
+			.OnCloseAnimationFinish(this, &SSidebar::OnCloseAnimationFinish)
+			.OnDrawerTargetSizeChanged(this, &SSidebar::OnDrawerTargetSizeChanged);
+	}
+
+	const TSharedRef<SSidebarDrawer> DrawerWidgetRef = InDrawer->DrawerWidget.ToSharedRef();
 
-	InDrawer->DrawerWidget->Open(bAnimateOpen);
+	if (ClosingDrawerWidgets.Contains(DrawerWidgetRef))
+	{
+		ClosingDrawerWidgets.Remove(DrawerWidgetRef);
+	}
+	else
+	{
+		DrawersOverlay->AddSlot()
+			.Padding(SlotPadding)
+			.HAlign(TabLocation == ESidebarTabLocation::Left ? HAlign_Left : HAlign_Right)
+			[
+				DrawerWidgetRef
+			];
+	}
+
+	OpenDrawerWidgets.Add(DrawerWidgetRef);
 
-	OpenedDrawers.Add(InDrawer->DrawerWidget.ToSharedRef());
+	InDrawer->DrawerWidget->Open(bInAnimate);
 
-	for (const TSharedRef<FSidebarDrawer>& DrawerTab : DrawerTabs)
+	for (const TSharedRef<FSidebarDrawer>& DrawerTab : Drawers)
 	{
 		DrawerTab->bIsOpen = false;
 	}
@@ -531,23 +561,28 @@ void SSidebar::OpenDrawerInternal(const TSharedRef<FSidebarDrawer>& InDrawer, co
 	FSlateApplication::Get().SetKeyboardFocus(InDrawer->DrawerWidget);
 }
 
-void SSidebar::CloseDrawerInternal(const TSharedRef<FSidebarDrawer>& InDrawer)
+void SSidebar::CloseDrawerInternal(const TSharedRef<FSidebarDrawer>& InDrawer, const bool bInAnimate)
 {
-	if (const TSharedPtr<SSidebarDrawer> OpenedDrawer = FindOpenedDrawer(InDrawer))
+	const TSharedPtr<SSidebarDrawer> FoundDrawerWidget = FindOpenDrawerWidget(InDrawer);
+	if (OpenDrawerWidgets.Contains(FoundDrawerWidget) && !ClosingDrawerWidgets.Contains(InDrawer->DrawerWidget))
 	{
-		const TSharedRef<SSidebarDrawer> OpenedDrawerRef = OpenedDrawer.ToSharedRef();
+		const TSharedRef<SSidebarDrawer> FoundDrawerWidgetRef = FoundDrawerWidget.ToSharedRef();
 
-		// Setting bInAnimateOpen to true causes a bug that forces keyboard focus to the drawer.
-		// Setting this false until this can be worked out. For some reason, inside Close(), the
-		// RegisterActiveTimer doesn't actually start a timer. This works for Open() though.
-		OpenedDrawerRef->Close(/*bInAnimateOpen=*/false);
+		FoundDrawerWidgetRef->Close(bInAnimate);
 
-		if (const TSharedPtr<SOverlay> DrawersOverlay = DrawersOverlayWeak.Pin())
+		if (bInAnimate)
 		{
-			DrawersOverlay->RemoveSlot(OpenedDrawerRef);
+			ClosingDrawerWidgets.Add(FoundDrawerWidgetRef);
+		}
+		else
+		{
+			if (const TSharedPtr<SOverlay> DrawersOverlay = DrawersOverlayWeak.Pin())
+			{
+				DrawersOverlay->RemoveSlot(FoundDrawerWidgetRef);
+			}
 		}
 
-		OpenedDrawers.Remove(OpenedDrawerRef);
+		OpenDrawerWidgets.Remove(FoundDrawerWidgetRef);
 
 		InDrawer->bIsOpen = false;
 	}
@@ -571,31 +606,30 @@ void SSidebar::SummonPinnedTabIfNothingOpened()
 	// - closing or restoring the current foreground tab
 	if (const TSharedPtr<FSidebarDrawer> PinnedTab = FindFirstPinnedTab())
 	{
-		OpenDrawerInternal(PinnedTab.ToSharedRef(), /*bAnimateOpen=*/true);
+		OpenDrawerInternal(PinnedTab.ToSharedRef());
 	}
 }
 
 void SSidebar::UpdateDrawerAppearance()
 {
-	TSharedPtr<FSidebarDrawer> OpenedTab;
-	if (OpenedDrawers.Num() > 0)
+	TSharedPtr<FSidebarDrawer> OpenedDrawer;
+	if (OpenDrawerWidgets.Num() > 0)
 	{
-		OpenedTab = OpenedDrawers.Last()->GetDrawer();
+		OpenedDrawer = OpenDrawerWidgets.Last()->GetDrawer();
 	}
 
-	for (const TSharedRef<FSidebarDrawer>& Drawer : DrawerTabs)
+	for (const TSharedRef<FSidebarDrawer>& Drawer : Drawers)
 	{
-		const TSharedPtr<SSidebarButton> TabButton = StaticCastSharedPtr<SSidebarButton>(Drawer->ButtonWidget);
-		if (TabButton.IsValid())
+		if (const TSharedPtr<SSidebarButton> TabButton = StaticCastSharedPtr<SSidebarButton>(Drawer->ButtonWidget))
 		{
-			TabButton->UpdateAppearance(OpenedTab);
+			TabButton->UpdateAppearance(OpenedDrawer);
 		}
 	}
 }
 
 TSharedPtr<FSidebarDrawer> SSidebar::FindDrawer(const FName InDrawerId) const
 {
-	const TSharedRef<FSidebarDrawer>* const FoundDrawer = DrawerTabs.FindByPredicate(
+	const TSharedRef<FSidebarDrawer>* const FoundDrawer = Drawers.FindByPredicate(
 		[InDrawerId](const TSharedRef<FSidebarDrawer>& InDrawer)
 		{
 			return InDrawerId == InDrawer->GetUniqueId();
@@ -605,7 +639,7 @@ TSharedPtr<FSidebarDrawer> SSidebar::FindDrawer(const FName InDrawerId) const
 
 TSharedPtr<FSidebarDrawer> SSidebar::FindFirstPinnedTab() const
 {
-	for (const TSharedRef<FSidebarDrawer>& DrawerTab : DrawerTabs)
+	for (const TSharedRef<FSidebarDrawer>& DrawerTab : Drawers)
 	{
 		if (DrawerTab->bIsPinned)
 		{
@@ -617,27 +651,27 @@ TSharedPtr<FSidebarDrawer> SSidebar::FindFirstPinnedTab() const
 
 TSharedPtr<FSidebarDrawer> SSidebar::GetForegroundTab() const
 {
-	const int32 Index = OpenedDrawers.FindLastByPredicate(
-		[](const TSharedRef<SSidebarDrawer>& Drawer)
+	const int32 Index = OpenDrawerWidgets.FindLastByPredicate(
+		[](const TSharedRef<SSidebarDrawer>& InDrawerWidget)
 		{
-			return Drawer->IsOpen() && !Drawer->IsClosing();
+			return InDrawerWidget->IsOpen() && !InDrawerWidget->IsClosing();
 		});
-	return Index == INDEX_NONE ? nullptr : OpenedDrawers[Index]->GetDrawer();
+	return Index == INDEX_NONE ? nullptr : OpenDrawerWidgets[Index]->GetDrawer();
 }
 
-TSharedPtr<SSidebarDrawer> SSidebar::FindOpenedDrawer(const TSharedRef<FSidebarDrawer>& InDrawer) const
+TSharedPtr<SSidebarDrawer> SSidebar::FindOpenDrawerWidget(const TSharedRef<FSidebarDrawer>& InDrawer) const
 {
-	const TSharedRef<SSidebarDrawer>* OpenedDrawer = OpenedDrawers.FindByPredicate(
+	const TSharedRef<SSidebarDrawer>* const OpenDrawWidget = OpenDrawerWidgets.FindByPredicate(
 		[&InDrawer](const TSharedRef<SSidebarDrawer>& Drawer)
 		{
 			return InDrawer == Drawer->GetDrawer();
 		});
-	return OpenedDrawer ? TSharedPtr<SSidebarDrawer>(*OpenedDrawer) : nullptr;
+	return OpenDrawWidget ? OpenDrawWidget->ToSharedPtr(): nullptr;
 }
 
 bool SSidebar::HasDrawerOpened() const
 {
-	for (const TSharedRef<FSidebarDrawer>& DrawerTab : DrawerTabs)
+	for (const TSharedRef<FSidebarDrawer>& DrawerTab : Drawers)
 	{
 		if (DrawerTab->bIsOpen)
 		{
@@ -649,7 +683,7 @@ bool SSidebar::HasDrawerOpened() const
 
 bool SSidebar::IsDrawerOpened(const FName InDrawerId) const
 {
-	for (const TSharedRef<FSidebarDrawer>& DrawerTab : DrawerTabs)
+	for (const TSharedRef<FSidebarDrawer>& DrawerTab : Drawers)
 	{
 		if (DrawerTab->bIsOpen && DrawerTab->GetUniqueId() == InDrawerId)
 		{
@@ -661,18 +695,18 @@ bool SSidebar::IsDrawerOpened(const FName InDrawerId) const
 
 FName SSidebar::GetOpenedDrawerId() const
 {
-	if (OpenedDrawers.IsEmpty())
+	if (OpenDrawerWidgets.IsEmpty())
 	{
 		return NAME_None;
 	}
 	
-	const TSharedRef<SSidebarDrawer> LastOpenedDrawer = OpenedDrawers.Last();
-	return LastOpenedDrawer->GetDrawer()->GetUniqueId();
+	const TSharedRef<SSidebarDrawer> LastOpenDrawerWidget = OpenDrawerWidgets.Last();
+	return LastOpenDrawerWidget->GetDrawer()->GetUniqueId();
 }
 
 bool SSidebar::HasDrawerPinned() const
 {
-	for (const TSharedRef<FSidebarDrawer>& DrawerTab : DrawerTabs)
+	for (const TSharedRef<FSidebarDrawer>& DrawerTab : Drawers)
 	{
 		if (DrawerTab->bIsPinned)
 		{
@@ -693,38 +727,44 @@ bool SSidebar::IsDrawerPinned(const FName InDrawerId) const
 
 void SSidebar::SetDrawerPinned(const FName InDrawerId, const bool bInIsPinned)
 {
-	const TSharedPtr<FSidebarDrawer> Drawer = FindDrawer(InDrawerId);
-	if (!Drawer.IsValid() || Drawer->bIsPinned == bInIsPinned)
+	const TSharedPtr<FSidebarDrawer> DrawerToPin = FindDrawer(InDrawerId);
+	if (!DrawerToPin.IsValid() || DrawerToPin->bIsPinned == bInIsPinned)
 	{
 		return;
 	}
 
 	if (bInIsPinned)
 	{
-		if (Drawer->bIsDocked)
+		if (DrawerToPin->bIsDocked)
 		{
 			SetDrawerDocked(InDrawerId, false);
 		}
 
-		if (!Drawer->bIsOpen)
+		if (!DrawerToPin->bIsOpen)
 		{
-			OpenDrawerInternal(Drawer.ToSharedRef(), false);
+			OpenDrawerInternal(DrawerToPin.ToSharedRef(), false);
 		}
-		if (!Drawer->bIsOpen)
+		if (!DrawerToPin->bIsOpen)
 		{
 			return;
 		}
+
+		// In case two modules attempt to register drawers with initially pinned states
+		for (const TSharedRef<FSidebarDrawer>& Drawer : Drawers)
+		{
+			Drawer->bIsPinned = false;
+		}
 	}
 
-	Drawer->bIsPinned = bInIsPinned;
+	DrawerToPin->bIsPinned = bInIsPinned;
 
 	if (bInIsPinned)
 	{
-		PinnedDrawerTabs.AddUnique(Drawer.ToSharedRef());
+		PinnedDrawerTabs.AddUnique(DrawerToPin.ToSharedRef());
 	}
 	else
 	{
-		PinnedDrawerTabs.Remove(Drawer.ToSharedRef());
+		PinnedDrawerTabs.Remove(DrawerToPin.ToSharedRef());
 	}
 }
 
@@ -757,34 +797,41 @@ void SSidebar::SetDrawerDocked(const FName InDrawerId, const bool bInIsDocked)
 		return;
 	}
 
-	const TSharedPtr<FSidebarDrawer> Drawer = FindDrawer(InDrawerId);
-	if (!Drawer.IsValid() || Drawer->bIsDocked == bInIsDocked)
+	const TSharedPtr<FSidebarDrawer> DrawerToDock = FindDrawer(InDrawerId);
+	if (!DrawerToDock.IsValid() || DrawerToDock->bIsDocked == bInIsDocked)
 	{
 		return;
 	}
 
 	if (bInIsDocked)
 	{
-		if (Drawer->bIsPinned)
+		if (DrawerToDock->bIsPinned)
 		{
 			SetDrawerPinned(InDrawerId, false);
 		}
 
-		CloseAllDrawers();
+		CloseAllDrawers(false);
 
 		if (DockedDrawerTab.IsValid())
 		{
 			UndockAllDrawers();
 		}
 
-		DockedDrawerTab = Drawer;
+		DockedDrawerTab = DrawerToDock;
+
+		// In case two modules attempt to register drawers with initially docked states
+		for (const TSharedRef<FSidebarDrawer>& Drawer : Drawers)
+		{
+			Drawer->bIsDocked = false;
+		}
 
+		DockedDrawerTab->bIsOpen = true;
 		DockedDrawerTab->bIsPinned = false;
 		DockedDrawerTab->bIsDocked = true;
 
-		if (Drawer->ContentWidget.IsValid())
+		if (DrawerToDock->ContentWidget.IsValid())
 		{
-			DockLocation->SetContent(Drawer->ContentWidget.ToSharedRef());
+			DockLocation->SetContent(DrawerToDock->ContentWidget.ToSharedRef());
 			
 			if (bHideWhenDocked)
 			{
@@ -805,6 +852,7 @@ void SSidebar::SetDrawerDocked(const FName InDrawerId, const bool bInIsDocked)
 	{
 		if (DockedDrawerTab.IsValid())
 		{
+			DockedDrawerTab->bIsOpen = false;
 			DockedDrawerTab->bIsDocked = false;
 			DockedDrawerTab.Reset();
 
@@ -821,7 +869,7 @@ void SSidebar::SetDrawerDocked(const FName InDrawerId, const bool bInIsDocked)
 
 void SSidebar::UndockAllDrawers()
 {
-	for (const TSharedRef<FSidebarDrawer>& DrawerTab : DrawerTabs)
+	for (const TSharedRef<FSidebarDrawer>& DrawerTab : Drawers)
 	{
 		SetDrawerDocked(DrawerTab->GetUniqueId(), false);
 	}
diff --git a/Engine/Source/Developer/ToolWidgets/Private/Sidebar/SSidebarButton.cpp b/Engine/Source/Developer/ToolWidgets/Private/Sidebar/SSidebarButton.cpp
--- a/Engine/Source/Developer/ToolWidgets/Private/Sidebar/SSidebarButton.cpp
+++ b/Engine/Source/Developer/ToolWidgets/Private/Sidebar/SSidebarButton.cpp
@@ -237,24 +237,30 @@ void SSidebarButton::Construct(const FArguments& InArgs, const TSharedRef<FSideb
 	UpdateAppearance(nullptr);
 }
 
-void SSidebarButton::UpdateAppearance(const TSharedPtr<FSidebarDrawer>& InOpenedDrawer)
+void SSidebarButton::UpdateAppearance(const TSharedPtr<FSidebarDrawer>& InLastDrawerOpen)
 {
+	const TSharedPtr<FSidebarDrawer> ThisDrawer = DrawerWeak.Pin();
+	if (!ThisDrawer.IsValid())
+	{
+		return;
+	}
+
 	float LabelRotation;
-	FName OpenBorderBrushName;
+	FName FocusBorderBrushName;
 
 	switch (TabLocation)
 	{
 	case ESidebarTabLocation::Left:
 		LabelRotation = -90.f;
-		OpenBorderBrushName = TEXT("Docking.Sidebar.Border_SquareRight");
+		FocusBorderBrushName = TEXT("Docking.Sidebar.Border_SquareRight");
 		break;
 	case ESidebarTabLocation::Right:
 		LabelRotation = 90.f;
-		OpenBorderBrushName = TEXT("Docking.Sidebar.Border_SquareLeft");
+		FocusBorderBrushName = TEXT("Docking.Sidebar.Border_SquareLeft");
 		break;
 	default:
 		LabelRotation = 0.f;
-		OpenBorderBrushName = TEXT("None");
+		FocusBorderBrushName = TEXT("None");
 		break;
 	}
 
@@ -263,17 +269,25 @@ void SSidebarButton::UpdateAppearance(const TSharedPtr<FSidebarDrawer>& InOpened
 		Label->SetRotation(LabelRotation);
 	}
 
-	if (InOpenedDrawer == DrawerWeak)
+	// Border when open/docked
+	if (InLastDrawerOpen == ThisDrawer && (!ThisDrawer->bIsDocked && ThisDrawer->bIsOpen))
 	{
-		// this button is the one with the tab that is actually opened so show the tab border
 		OpenBorder->SetVisibility(EVisibility::HitTestInvisible);
-		MainButton->SetButtonStyle(&FAppStyle::Get().GetWidgetStyle<FButtonStyle>(TEXT("Docking.SidebarButton.Opened")));
-
-		OpenBorder->SetBorderImage(FAppStyle::Get().GetBrush(OpenBorderBrushName));
+		OpenBorder->SetBorderImage(FAppStyle::Get().GetBrush(FocusBorderBrushName));
 	}
 	else
 	{
 		OpenBorder->SetVisibility(EVisibility::Collapsed);
+	}
+
+	// Button style
+	if (InLastDrawerOpen == ThisDrawer || ThisDrawer->bIsDocked)
+	{
+		// this button is the one with the tab that is actually opened so show the tab border
+		MainButton->SetButtonStyle(&FAppStyle::Get().GetWidgetStyle<FButtonStyle>(TEXT("Docking.SidebarButton.Opened")));
+	}
+	else
+	{
 		MainButton->SetButtonStyle(&FAppStyle::Get().GetWidgetStyle<FButtonStyle>(TEXT("Docking.SidebarButton.Closed")));
 	}
 }
diff --git a/Engine/Source/Developer/ToolWidgets/Private/Sidebar/SSidebarButton.h b/Engine/Source/Developer/ToolWidgets/Private/Sidebar/SSidebarButton.h
--- a/Engine/Source/Developer/ToolWidgets/Private/Sidebar/SSidebarButton.h
+++ b/Engine/Source/Developer/ToolWidgets/Private/Sidebar/SSidebarButton.h
@@ -42,7 +42,7 @@ public:
 
 	void Construct(const FArguments& InArgs, const TSharedRef<FSidebarDrawer>& InDrawer, const ESidebarTabLocation InTabLocation);
 
-	void UpdateAppearance(const TSharedPtr<FSidebarDrawer>& InOpenedDrawer);
+	void UpdateAppearance(const TSharedPtr<FSidebarDrawer>& InLastDrawerOpen);
 
 	void OnTabRenamed(const TWeakPtr<FSidebarDrawer>& InDrawer);
 	
diff --git a/Engine/Source/Developer/ToolWidgets/Private/Sidebar/SSidebarDrawer.cpp b/Engine/Source/Developer/ToolWidgets/Private/Sidebar/SSidebarDrawer.cpp
--- a/Engine/Source/Developer/ToolWidgets/Private/Sidebar/SSidebarDrawer.cpp
+++ b/Engine/Source/Developer/ToolWidgets/Private/Sidebar/SSidebarDrawer.cpp
@@ -31,7 +31,8 @@ void SSidebarDrawer::Construct(const FArguments& InArgs, const TSharedRef<FSideb
 	OnDrawerTargetSizeChanged = InArgs._OnDrawerTargetSizeChanged;
 	OnDrawerFocused = InArgs._OnDrawerFocused;
 	OnDrawerFocusLost = InArgs._OnDrawerFocusLost;
-	OnDrawerClosed = InArgs._OnDrawerClosed;
+	OnCloseAnimationFinish = InArgs._OnCloseAnimationFinish;
+	OnOpenAnimationFinish = InArgs._OnOpenAnimationFinish;
 
 	OpenCloseAnimation = FCurveSequence(0.f, AnimationLength, ECurveEaseFunction::QuadOut);
 
@@ -61,8 +62,12 @@ void SSidebarDrawer::Open(const bool bInAnimateOpen)
 		return;
 	}
 
-	const float StartTime = OpenCloseAnimation.IsPlaying() ? OpenCloseAnimation.GetSequenceTime() : 0.f;
-	OpenCloseAnimation.Play(AsShared(), false, StartTime, false);
+	if (OpenCloseAnimation.IsInReverse())
+	{
+		OpenCloseAnimation.Reverse();
+	}
+
+	OpenCloseAnimation.Play(AsShared(), false, OpenCloseAnimation.GetSequenceTime(), false);
 
 	if (!OpenCloseTimer.IsValid())
 	{
@@ -80,8 +85,10 @@ void SSidebarDrawer::Close(const bool bInAnimateOpen)
 		return;
 	}
 
-	const float StartTime = OpenCloseAnimation.IsPlaying() ? OpenCloseAnimation.GetSequenceTime() : AnimationLength;
-	OpenCloseAnimation.PlayReverse(AsShared(), false, StartTime, true);
+	if (OpenCloseAnimation.IsForward())
+	{
+		OpenCloseAnimation.Reverse();
+	}
 
 	if (!OpenCloseTimer.IsValid())
 	{
@@ -497,7 +504,11 @@ EActiveTimerReturnType SSidebarDrawer::UpdateAnimation(const double InCurrentTim
 	{
 		if (OpenCloseAnimation.IsAtStart())
 		{
-			OnDrawerClosed.ExecuteIfBound(SharedThis(this));
+			OnCloseAnimationFinish.ExecuteIfBound(SharedThis(this));
+		}
+		else if (OpenCloseAnimation.IsAtEnd())
+		{
+			OnOpenAnimationFinish.ExecuteIfBound(SharedThis(this));
 		}
 
 		FSlateThrottleManager::Get().LeaveResponsiveMode(AnimationThrottle);
@@ -576,7 +587,7 @@ void SSidebarDrawer::OnGlobalFocusChanging(const FFocusEvent& InFocusEvent
 
 				// See if the menu being opened is owned by the drawer contents and if so the menu should not be dismissed
 				FSlateApplication::Get().GeneratePathToWidgetUnchecked(MenuHost.ToSharedRef(), MenuHostPath);
-				
+
 				if (!MenuHostPath.ContainsWidget(&ChildSlot.GetWidget().Get()))
 				{
 					bShouldLoseFocus = true;
diff --git a/Engine/Source/Developer/ToolWidgets/Private/Sidebar/SSidebarDrawer.h b/Engine/Source/Developer/ToolWidgets/Private/Sidebar/SSidebarDrawer.h
--- a/Engine/Source/Developer/ToolWidgets/Private/Sidebar/SSidebarDrawer.h
+++ b/Engine/Source/Developer/ToolWidgets/Private/Sidebar/SSidebarDrawer.h
@@ -45,8 +45,10 @@ public:
 		SLATE_EVENT(FGenericSidebarDrawerWidgetDelegate, OnDrawerFocused)
 		/** Called when the drawer loses focus */
 		SLATE_EVENT(FGenericSidebarDrawerWidgetDelegate, OnDrawerFocusLost)
+		/** Called when the drawer is completely opened (i.e will be called once the open animation completes */
+		SLATE_EVENT(FGenericSidebarDrawerWidgetDelegate, OnOpenAnimationFinish)
 		/** Called when the drawer is completely closed (i.e will be called once the close animation completes */
-		SLATE_EVENT(FGenericSidebarDrawerWidgetDelegate, OnDrawerClosed)
+		SLATE_EVENT(FGenericSidebarDrawerWidgetDelegate, OnCloseAnimationFinish)
 	SLATE_END_ARGS()
 
 	~SSidebarDrawer();
@@ -116,7 +118,8 @@ protected:
 	FOnSidebarDrawerTargetSizeChanged OnDrawerTargetSizeChanged;
 	FGenericSidebarDrawerWidgetDelegate OnDrawerFocused;
 	FGenericSidebarDrawerWidgetDelegate OnDrawerFocusLost;
-	FGenericSidebarDrawerWidgetDelegate OnDrawerClosed;
+	FGenericSidebarDrawerWidgetDelegate OnCloseAnimationFinish;
+	FGenericSidebarDrawerWidgetDelegate OnOpenAnimationFinish;
 
 	FCurveSequence OpenCloseAnimation;
 	TSharedPtr<FActiveTimerHandle> OpenCloseTimer;
diff --git a/Engine/Source/Developer/ToolWidgets/Public/Sidebar/SSidebar.h b/Engine/Source/Developer/ToolWidgets/Public/Sidebar/SSidebar.h
--- a/Engine/Source/Developer/ToolWidgets/Public/Sidebar/SSidebar.h
+++ b/Engine/Source/Developer/ToolWidgets/Public/Sidebar/SSidebar.h
@@ -153,7 +153,7 @@ public:
 	bool TryOpenDrawer(const FName InDrawerId);
 
 	/** Closes any drawers that are open. */
-	void CloseAllDrawers();
+	void CloseAllDrawers(const bool bInAnimate = true);
 
 	/** @return True if the sidebar has any drawer that is opened. */
 	bool HasDrawerOpened() const;
@@ -230,7 +230,8 @@ private:
 	void OnDrawerTabPinToggled(const TSharedRef<FSidebarDrawer>& InDrawer, const bool bIsPinned);
 	void OnDrawerTabDockToggled(const TSharedRef<FSidebarDrawer>& InDrawer, const bool bIsDocked);
 	void OnTabDrawerFocusLost(const TSharedRef<SSidebarDrawer>& InDrawerWidget);
-	void OnTabDrawerClosed(const TSharedRef<SSidebarDrawer>& InDrawerWidget);
+	void OnOpenAnimationFinish(const TSharedRef<SSidebarDrawer>& InDrawerWidget);
+	void OnCloseAnimationFinish(const TSharedRef<SSidebarDrawer>& InDrawerWidget);
 	void OnDrawerTargetSizeChanged(const TSharedRef<SSidebarDrawer>& InDrawerWidget, const float InNewSize);
 	TSharedRef<SWidget> OnGetTabDrawerContextMenuWidget(TSharedRef<FSidebarDrawer> InDrawer);
 	void BuildOptionsMenu(UToolMenu* const InMenu);
@@ -241,10 +242,10 @@ private:
 	/** Removes all drawers instantly (including drawers for pinned tabs). */
 	void RemoveAllDrawers();
 
-	EActiveTimerReturnType OnOpenPendingDrawerTimer(const double CurrentTime, const float DeltaTime);
-	void OpenDrawerNextFrame(const TSharedRef<FSidebarDrawer>& InDrawer, const bool bAnimateOpen);
-	void OpenDrawerInternal(const TSharedRef<FSidebarDrawer>& InDrawer, const bool bAnimateOpen);
-	void CloseDrawerInternal(const TSharedRef<FSidebarDrawer>& InDrawer);
+	EActiveTimerReturnType OnOpenPendingDrawerTimer(const double InCurrentTime, const float InDeltaTime);
+	void OpenDrawerNextFrame(const TSharedRef<FSidebarDrawer>& InDrawer, const bool bInAnimate = true);
+	void OpenDrawerInternal(const TSharedRef<FSidebarDrawer>& InDrawer, const bool bInAnimate = true);
+	void CloseDrawerInternal(const TSharedRef<FSidebarDrawer>& InDrawer, const bool bInAnimate = true);
 
 	/** Reopens the pinned tab only if there are no other open drawers. This should be used to bring pinned tabs back after other tabs lose focus/are closed. */
 	void SummonPinnedTabIfNothingOpened();
@@ -261,7 +262,7 @@ private:
 	TSharedPtr<FSidebarDrawer> GetForegroundTab() const;
 
 	/** Returns the drawer for the given tab if it's open. */
-	TSharedPtr<SSidebarDrawer> FindOpenedDrawer(const TSharedRef<FSidebarDrawer>& InDrawer) const;
+	TSharedPtr<SSidebarDrawer> FindOpenDrawerWidget(const TSharedRef<FSidebarDrawer>& InDrawer) const;
 
 	TWeakPtr<SOverlay> DrawersOverlayWeak;
 	TWeakPtr<SBox> DockLocationWeak;
@@ -273,12 +274,16 @@ private:
 	bool bDisableDock = false;
 	FOnSidebarDrawerDockStateChanged OnDockStateChanged;
 
-	TSharedPtr<SScrollBox> TabContainer;
+	TSharedPtr<SScrollBox> TabButtonContainer;
 
-	TArray<TSharedRef<FSidebarDrawer>> DrawerTabs;
+	TArray<TSharedRef<FSidebarDrawer>> Drawers;
 
-	/** Generally speaking one drawer is only ever open at once but we animate any previous drawer closing so there could be more than one while an animation is playing. */
-	TArray<TSharedRef<SSidebarDrawer>> OpenedDrawers;
+	/** Generally speaking one drawer is only ever open at once but we animate any previous drawer
+	 * closing so there could be more than one while an animation is playing. A docked drawer is
+	 * also considered open, along with any user opened/pinned drawers. */
+	TArray<TSharedRef<SSidebarDrawer>> OpenDrawerWidgets;
+
+	TArray<TSharedRef<SSidebarDrawer>> ClosingDrawerWidgets;
 
 	TWeakPtr<FSidebarDrawer> PendingTabToOpen;
 	bool bAnimatePendingTabOpen = false;
