[Sequencer] Fix sequencer sidebar focus issue when double clicking color keys

diff --git a/Engine/Source/Developer/ToolWidgets/Private/Sidebar/SSidebarDrawer.cpp b/Engine/Source/Developer/ToolWidgets/Private/Sidebar/SSidebarDrawer.cpp
--- a/Engine/Source/Developer/ToolWidgets/Private/Sidebar/SSidebarDrawer.cpp
+++ b/Engine/Source/Developer/ToolWidgets/Private/Sidebar/SSidebarDrawer.cpp
@@ -543,6 +543,15 @@ void SSidebarDrawer::OnGlobalFocusChanging(const FFocusEvent& InFocusEvent
 	{
 		return;
 	}
+	TGuardValue<bool> ReEntrancyGuard(bIsReEntrant, true);
+
+	const TSharedPtr<FSidebarDrawer> Drawer = DrawerWeak.Pin();
+
+	// Only open drawers that are not docked or pinned need to close the drawer when focus is lost
+	if (!Drawer || Drawer->bIsDocked || Drawer->bIsPinned || !Drawer->bIsOpen)
+	{
+		return;
+	}
 
 	// Do not close due to slow tasks as those opening send window activation events
 	if (GIsSlowTask || FSlateApplication::Get().GetActiveModalWindow().IsValid())
@@ -550,17 +559,12 @@ void SSidebarDrawer::OnGlobalFocusChanging(const FFocusEvent& InFocusEvent
 		return;
 	}
 
-	TGuardValue<bool> ReEntrancyGuard(bIsReEntrant, true);
-
 	const TSharedRef<SSidebarDrawer> ThisWidget = SharedThis(this);
 
 	TArray<TSharedRef<SWidget>, TInlineAllocator<4>> LegalFocusWidgets;
 	LegalFocusWidgets.Add(ThisWidget);
 	LegalFocusWidgets.Add(ChildSlot.GetWidget());
-	if (const TSharedPtr<FSidebarDrawer> Drawer = DrawerWeak.Pin())
-	{
-		LegalFocusWidgets.Add(Drawer->ButtonWidget.ToSharedRef());
-	}
+	LegalFocusWidgets.Add(Drawer->ButtonWidget.ToSharedRef());
 
 	bool bShouldLoseFocus = false;
 
diff --git a/Engine/Source/Editor/MovieSceneTools/Private/Sections/MovieSceneSectionHelpers.cpp b/Engine/Source/Editor/MovieSceneTools/Private/Sections/MovieSceneSectionHelpers.cpp
--- a/Engine/Source/Editor/MovieSceneTools/Private/Sections/MovieSceneSectionHelpers.cpp
+++ b/Engine/Source/Editor/MovieSceneTools/Private/Sections/MovieSceneSectionHelpers.cpp
@@ -139,6 +139,7 @@ FMovieSceneKeyColorPicker::FMovieSceneKeyColorPicker(UMovieSceneSection* Section
 	PickerArgs.DisplayGamma = TAttribute<float>::Create(TAttribute<float>::FGetter::CreateUObject(GEngine, &UEngine::GetDisplayGamma));
 	PickerArgs.InitialColor = InitialColor;
 	PickerArgs.ParentWidget = FSlateApplication::Get().GetActiveTopLevelWindow();
+	PickerArgs.OptionalOwningDetailsView = PickerArgs.ParentWidget;
 
 	PickerArgs.OnColorCommitted = FOnLinearColorValueChanged::CreateRaw(this, &FMovieSceneKeyColorPicker::OnColorPickerPicked, RChannel, GChannel, BChannel, AChannel);
 	PickerArgs.OnColorPickerWindowClosed = FOnWindowClosed::CreateRaw(this, &FMovieSceneKeyColorPicker::OnColorPickerClosed, Section, RChannel, GChannel, BChannel, AChannel);
