Modify hash function to compute the avatar's color

In order for other plugins/apps to be able to match UE's avatar color it is necessary to have a easy, portable and simple hashing function.

diff --git a/Engine/Source/Developer/ToolWidgets/Private/SAvatar.cpp b/Engine/Source/Developer/ToolWidgets/Private/SAvatar.cpp
--- a/Engine/Source/Developer/ToolWidgets/Private/SAvatar.cpp
+++ b/Engine/Source/Developer/ToolWidgets/Private/SAvatar.cpp
@@ -1,6 +1,7 @@
 // Copyright Epic Games, Inc. All Rights Reserved.
 
 #include "SAvatar.h"
+#include "Containers/StringConv.h"
 #include "Framework/Application/SlateApplication.h"
 #include "Widgets/Layout/SBox.h"
 #include "Styling/CoreStyle.h"
@@ -14,11 +15,7 @@ void SAvatar::Construct(const FArguments& InArgs)
 
 	bShowInitial = InArgs._ShowInitial;
 
-	const uint32 Hash = GetTypeHash(Identifier);
-	const uint8 H = ((Hash >> 0) & 0xFF) ^ ((Hash >> 8) & 0xFF) ^ ((Hash >> 16) & 0xFF) ^ ((Hash >> 24) & 0xFF);
-	const uint8 S = 128 + ((Hash >> 8) & 0x7F);
-	const uint8 V = 128 + ((Hash >> 16) & 0x7F) / 2;
-	BackgroundColor = FLinearColor::MakeFromHSV8(H, S, V).ToFColor(false);
+	BackgroundColor = ComputeBackgroundColor();
 	ForegroundColor = FColor::White;
 
 	ChildSlot
@@ -69,4 +66,28 @@ int32 SAvatar::OnPaint(const FPaintArgs& Args, const FGeometry& AllottedGeometry
 	}
 
 	return SCompoundWidget::OnPaint(Args, AllottedGeometry, MyCullingRect, OutDrawElements, LayerId, InWidgetStyle, bParentEnabled && IsEnabled());
+}
+
+uint32 SAvatar::Hash() const
+{
+	// Warning: changing the hash computation could create discrepancies with other Epic owned systems' avatar
+	// appearance.
+	uint32 HashValue = 5381;
+	using FStringUtf8 = TStringConversion<TStringConvert<FString::ElementType, UTF8CHAR>>;
+	const FStringUtf8 IdentifierUtf8 = StringCast<UTF8CHAR>(*Identifier);
+	const FUtf8StringView IdentifierUtf8View(IdentifierUtf8.Get(), IdentifierUtf8.Length());
+	for (const UTF8CHAR Byte : IdentifierUtf8View)
+	{
+		HashValue = ((HashValue << 5) + HashValue) + Byte;
+	}
+	return HashValue;
+}
+
+FColor SAvatar::ComputeBackgroundColor() const
+{
+	const uint64 Value = Hash();
+	const uint8 H = ((Value >> 0) & 0xFF) ^ ((Value >> 8) & 0xFF) ^ ((Value >> 16) & 0xFF) ^ ((Value >> 24) & 0xFF);
+	const uint8 S = 128 + ((Value >> 8) & 0x7F);
+	const uint8 V = 128 + ((Value >> 16) & 0x7F) / 2;
+	return FLinearColor::MakeFromHSV8(H, S, V).ToFColor(false);
 }
\ No newline at end of file
diff --git a/Engine/Source/Developer/ToolWidgets/Public/SAvatar.h b/Engine/Source/Developer/ToolWidgets/Public/SAvatar.h
--- a/Engine/Source/Developer/ToolWidgets/Public/SAvatar.h
+++ b/Engine/Source/Developer/ToolWidgets/Public/SAvatar.h
@@ -32,6 +32,19 @@ public:
 	virtual int32 OnPaint(const FPaintArgs& Args, const FGeometry& AllottedGeometry, const FSlateRect& MyCullingRect, FSlateWindowElementList& OutDrawElements, int32 LayerId, const FWidgetStyle& InWidgetStyle, bool bParentEnabled) const override;
 
 private:
+	/** Computes a hash value using the djb2 algorithm.
+	 * @return the hash value of the Identifier
+	 */
+	uint32 Hash() const;
+	
+	/** Computes the avatar's background color based on the Identifier
+	 * Calls the Hash() function above to calculate an HSV color which is
+	 * transformed into an FColor that ultimately represents the user's avatar color.
+	 * 
+	 * @return the RGB representation of the Identifier
+	 */
+	FColor ComputeBackgroundColor() const;
+
 	/** The Identifier of the avatar */
 	FString Identifier;
 
