[ToolWidgets] SActionButton uses Icon binding, if provided

diff --git a/Engine/Source/Developer/ToolWidgets/Private/SActionButton.cpp b/Engine/Source/Developer/ToolWidgets/Private/SActionButton.cpp
--- a/Engine/Source/Developer/ToolWidgets/Private/SActionButton.cpp
+++ b/Engine/Source/Developer/ToolWidgets/Private/SActionButton.cpp
@@ -70,10 +70,10 @@ void SActionButton::Construct(const FArguments& InArgs)
 	// Check for widget level override, then style override, otherwise unset
 	const TAttribute<const FSlateBrush*> Icon =
 		InArgs._Icon.IsSet()
-		? InArgs._Icon.Get()
+		? InArgs._Icon
 		: ActionButtonStyle->IconBrush.IsSet()
 		? &ActionButtonStyle->IconBrush.GetValue()
-		: nullptr;
+		: TAttribute<const FSlateBrush*>();
 
 	const bool bHasIcon = Icon.Get() || Icon.IsBound();
 
