[BUGFIX] Button variants will correctly use delegates if set.

diff --git a/Engine/Source/Developer/ToolWidgets/Private/SPrimaryButton.cpp b/Engine/Source/Developer/ToolWidgets/Private/SPrimaryButton.cpp
--- a/Engine/Source/Developer/ToolWidgets/Private/SPrimaryButton.cpp
+++ b/Engine/Source/Developer/ToolWidgets/Private/SPrimaryButton.cpp
@@ -15,7 +15,7 @@ void SPrimaryButton::Construct(const FArguments& InArgs)
 
 	// Check for widget level override, then style override, otherwise unset
 	const TAttribute<const FSlateBrush*> Icon = InArgs._Icon.IsSet()
-		? InArgs._Icon.Get()
+		? InArgs._Icon
 		: ActionButtonStyle->IconBrush.IsSet()
 		? &ActionButtonStyle->IconBrush.GetValue()
 		: nullptr;
diff --git a/Engine/Source/Developer/ToolWidgets/Private/SSimpleButton.cpp b/Engine/Source/Developer/ToolWidgets/Private/SSimpleButton.cpp
--- a/Engine/Source/Developer/ToolWidgets/Private/SSimpleButton.cpp
+++ b/Engine/Source/Developer/ToolWidgets/Private/SSimpleButton.cpp
@@ -14,7 +14,7 @@ void SSimpleButton::Construct(const FArguments& InArgs)
 
 	// Check for widget level override, then style override, otherwise unset
 	const TAttribute<const FSlateBrush*> Icon = InArgs._Icon.IsSet()
-		? InArgs._Icon.Get()
+		? InArgs._Icon
 		: ActionButtonStyle->IconBrush.IsSet()
 		? &ActionButtonStyle->IconBrush.GetValue()
 		: nullptr;
diff --git a/Engine/Source/Developer/ToolWidgets/Private/SSimpleComboButton.cpp b/Engine/Source/Developer/ToolWidgets/Private/SSimpleComboButton.cpp
--- a/Engine/Source/Developer/ToolWidgets/Private/SSimpleComboButton.cpp
+++ b/Engine/Source/Developer/ToolWidgets/Private/SSimpleComboButton.cpp
@@ -15,7 +15,7 @@ void SSimpleComboButton::Construct(const FArguments& InArgs)
 
 	// Check for widget level override, then style override, otherwise unset
 	const TAttribute<const FSlateBrush*> Icon = InArgs._Icon.IsSet()
-		? InArgs._Icon.Get()
+		? InArgs._Icon
 		: ActionButtonStyle->IconBrush.IsSet()
 		? &ActionButtonStyle->IconBrush.GetValue()
 		: nullptr;
diff --git a/Engine/Source/Developer/ToolWidgets/Private/ToolWidgetsUtilitiesPrivate.cpp b/Engine/Source/Developer/ToolWidgets/Private/ToolWidgetsUtilitiesPrivate.cpp
--- a/Engine/Source/Developer/ToolWidgets/Private/ToolWidgetsUtilitiesPrivate.cpp
+++ b/Engine/Source/Developer/ToolWidgets/Private/ToolWidgetsUtilitiesPrivate.cpp
@@ -86,9 +86,8 @@ namespace UE::ToolWidgets::Private
 			check(InActionButtonStyle);
 
 			// Check for widget level override, then style override, otherwise unset
-			const TAttribute<const FSlateBrush*> Icon =
-				InIcon.IsSet()
-				? InIcon.Get()
+			const TAttribute<const FSlateBrush*> Icon = InIcon.IsSet()
+				? InIcon
 				: InActionButtonStyle->IconBrush.IsSet()
 				? &InActionButtonStyle->IconBrush.GetValue()
 				: nullptr;
