[Basic Filter Bar] Making sure combo button embeded in a menu widget can close its menu when user clicks on a parent menu

diff --git a/Engine/Source/Developer/ToolWidgets/Public/Filters/SBasicFilterBar.h b/Engine/Source/Developer/ToolWidgets/Public/Filters/SBasicFilterBar.h
--- a/Engine/Source/Developer/ToolWidgets/Public/Filters/SBasicFilterBar.h
+++ b/Engine/Source/Developer/ToolWidgets/Public/Filters/SBasicFilterBar.h
@@ -267,6 +267,8 @@ public:
 				.ToolTipText(LOCTEXT("AddFilterToolTip", "Open the Add Filter Menu to add or manage filters."))
 				.OnGetMenuContent(InFilterBar, &SBasicFilterBar<FilterType>::MakeAddFilterMenu)
 				.ContentPadding(FMargin(1, 0))
+ 				// If the filter bar is embedded in a parent menu, this allows the Add Filter Menu to close when clicking on an empty area in the parent menu
+ 				.CollapseMenuOnParentFocus(true)
 				.ButtonContent()
 				[
 					FilterImage.ToSharedRef()
