Fixed issue with shared reference being held by both the sidebar container and the sidebar itself, causing Sequencer Widget to not properly destroy itself in Motion Design

diff --git a/Engine/Plugins/Experimental/Avalanche/Source/AvalancheSequencer/Private/AvaSequencer.cpp b/Engine/Plugins/Experimental/Avalanche/Source/AvalancheSequencer/Private/AvaSequencer.cpp
--- a/Engine/Plugins/Experimental/Avalanche/Source/AvalancheSequencer/Private/AvaSequencer.cpp
+++ b/Engine/Plugins/Experimental/Avalanche/Source/AvalancheSequencer/Private/AvaSequencer.cpp
@@ -1499,29 +1499,26 @@ TSharedRef<SWidget> FAvaSequencer::CreateSequenceWidget()
 	check(IsValid(SequencerSettings));
 	FSidebarState& SidebarState = SequencerSettings->GetSidebarState();
 
-	TSharedPtr<SWidget> OutWidget;
+	// Make sure the sequence tree widget is created
+	GetSequenceTreeWidget();
 
-	if (SidebarState.IsVisible())
+	if (!SidebarState.IsVisible())
 	{
-		SidebarContainer = SNew(SSidebarContainer);
+		return Sequencer->GetSequencerWidget();
+	}
 
-		LeftSidebar = SNew(SSidebar, SidebarContainer.ToSharedRef())
-			.TabLocation(ESidebarTabLocation::Left)
-			.InitialDrawerSize(SidebarState.GetDrawerSize())
-			.OnStateChanged(this, &FAvaSequencer::OnSidebarStateChanged)
-			.OnGetContent(FOnGetContent::CreateLambda([this]()
-				{
-					return GetSequencer()->GetSequencerWidget();
-				}));
+	TSharedRef<SSidebarContainer> SidebarContainer = SNew(SSidebarContainer);
 
-		SidebarContainer->RebuildSidebar(LeftSidebar.ToSharedRef(), SidebarState);
+	TSharedRef<SSidebar> LeftSidebar = SNew(SSidebar, SidebarContainer)
+		.TabLocation(ESidebarTabLocation::Left)
+		.InitialDrawerSize(SidebarState.GetDrawerSize())
+		.OnStateChanged(this, &FAvaSequencer::OnSidebarStateChanged)
+		.OnGetContent(FOnGetContent::CreateLambda([this]()
+		{
+			return GetSequencer()->GetSequencerWidget();
+		}));
 
-		OutWidget = SidebarContainer.ToSharedRef();
-	}
-	else
-	{
-		OutWidget = Sequencer->GetSequencerWidget();
-	}
+	SidebarContainer->RebuildSidebar(LeftSidebar, SidebarState);
 
 	FSidebarDrawerConfig SequenceTreeDrawerConfig;
 	SequenceTreeDrawerConfig.UniqueId = SidebarDrawerId;
@@ -1532,15 +1529,12 @@ TSharedRef<SWidget> FAvaSequencer::CreateSequenceWidget()
 
 	LeftSidebar->RegisterDrawer(MoveTemp(SequenceTreeDrawerConfig));
 
-	// Make sure the sequence tree widget is created
-	GetSequenceTreeWidget();
-
 	const TSharedRef<FAvaSequencer> ThisSequencerRef = SharedThis(this);
 	LeftSidebar->RegisterDrawerSection(SidebarDrawerId, MakeShared<FAvaSequenceTreeDetails>(ThisSequencerRef));
 	LeftSidebar->RegisterDrawerSection(SidebarDrawerId, MakeShared<FAvaSequencePlaybackDetails>(ThisSequencerRef));
 	LeftSidebar->RegisterDrawerSection(SidebarDrawerId, MakeShared<FAvaSequenceSettingsDetails>(ThisSequencerRef));
 
-	return OutWidget.ToSharedRef();
+	return SidebarContainer;
 }
 
 void FAvaSequencer::OnActorsCopied(FString& InOutCopiedData, TConstArrayView<AActor*> InCopiedActors)
diff --git a/Engine/Plugins/Experimental/Avalanche/Source/AvalancheSequencer/Private/AvaSequencer.h b/Engine/Plugins/Experimental/Avalanche/Source/AvalancheSequencer/Private/AvaSequencer.h
--- a/Engine/Plugins/Experimental/Avalanche/Source/AvalancheSequencer/Private/AvaSequencer.h
+++ b/Engine/Plugins/Experimental/Avalanche/Source/AvalancheSequencer/Private/AvaSequencer.h
@@ -283,9 +283,7 @@ private:
 
 	TSharedPtr<FAvaEaseCurveTool> EaseCurveTool;
 
-	TSharedPtr<SSidebarContainer> SidebarContainer;
-	TSharedPtr<SSidebar> LeftSidebar;
-
 	TSharedPtr<FExtender> SidebarExtender;
+
 	FDelegateHandle SidebarSelectionExtenderHandle;
 };
diff --git a/Engine/Source/Developer/ToolWidgets/Private/Sidebar/SSidebar.cpp b/Engine/Source/Developer/ToolWidgets/Private/Sidebar/SSidebar.cpp
--- a/Engine/Source/Developer/ToolWidgets/Private/Sidebar/SSidebar.cpp
+++ b/Engine/Source/Developer/ToolWidgets/Private/Sidebar/SSidebar.cpp
@@ -20,12 +20,11 @@
 
 SSidebar::~SSidebar()
 {
-	RemoveAllDrawers();
 }
 
 void SSidebar::Construct(const FArguments& InArgs, const TSharedRef<SSidebarContainer>& InContainerWidget)
 {
-	ContainerWidget = InContainerWidget;
+	ContainerWidgetWeak = InContainerWidget;
 
 	TabLocation = InArgs._TabLocation;
 	InContainerWidget->SidebarSizePercent = InArgs._InitialDrawerSize;
@@ -101,7 +100,11 @@ bool SSidebar::RegisterDrawer(FSidebarDrawerConfig&& InDrawerConfig)
 		SetDrawerDocked(DrawerId, true);
 	}
 
-	ContainerWidget->UpdateDrawerTabAppearance();
+	TSharedPtr<SSidebarContainer> ContainerWidget = ContainerWidgetWeak.Pin();
+	if (ensure(ContainerWidget.IsValid()))
+	{
+		ContainerWidget->UpdateDrawerTabAppearance();
+	}
 
 	if (bHideWhenAllDocked && !AreAllDrawersDocked())
 	{
@@ -113,11 +116,12 @@ bool SSidebar::RegisterDrawer(FSidebarDrawerConfig&& InDrawerConfig)
 
 bool SSidebar::UnregisterDrawer(const FName InDrawerId)
 {
-	if (IsDrawerOpened(InDrawerId))
+	TSharedPtr<SSidebarContainer> ContainerWidget = ContainerWidgetWeak.Pin();
+	if (ensure(ContainerWidget.IsValid()) && IsDrawerOpened(InDrawerId))
 	{
 		ContainerWidget->CloseAllDrawerWidgets(false);
 	}
-	
+
 	const int32 IndexToRemove = Drawers.IndexOfByPredicate(
 		[InDrawerId](const TSharedRef<FSidebarDrawer>& InDrawer)
 		{
@@ -133,7 +137,10 @@ bool SSidebar::UnregisterDrawer(const FName InDrawerId)
 	RemoveDrawer(Drawers[IndexToRemove]);
 	Drawers.RemoveAt(IndexToRemove);
 
-	ContainerWidget->SummonPinnedTabIfNothingOpened();
+	if (ContainerWidget.IsValid())
+	{
+		ContainerWidget->SummonPinnedTabIfNothingOpened();
+	}
 
 	// Clear the pinned flag when the tab is removed from the sidebar.
 	// (Users probably expect that pinning a tab, restoring it/closing it,
@@ -144,7 +151,7 @@ bool SSidebar::UnregisterDrawer(const FName InDrawerId)
 	{
 		SetVisibility(EVisibility::Collapsed);
 	}
-	else
+	else if (ContainerWidget.IsValid())
 	{
 		ContainerWidget->UpdateDrawerTabAppearance();
 	}
@@ -218,31 +225,44 @@ bool SSidebar::TryOpenDrawer(const FName InDrawerId)
 		return false;
 	}
 
-	ContainerWidget->OpenDrawerNextFrame(Drawer.ToSharedRef());
+	TSharedPtr<SSidebarContainer> ContainerWidget = ContainerWidgetWeak.Pin();
+	if (ensure(ContainerWidget.IsValid()))
+	{
+		ContainerWidget->OpenDrawerNextFrame(Drawer.ToSharedRef());
+	}
 
 	return true;
 }
 
 void SSidebar::CloseAllDrawers(const bool bInAnimate)
 {
-	ContainerWidget->CloseAllDrawerWidgets(bInAnimate);
+	TSharedPtr<SSidebarContainer> ContainerWidget = ContainerWidgetWeak.Pin();
+	if (ensure(ContainerWidget.IsValid()))
+	{
+		ContainerWidget->CloseAllDrawerWidgets(bInAnimate);
+	}
 }
 
 void SSidebar::OnTabDrawerButtonPressed(const TSharedRef<FSidebarDrawer>& InDrawer)
 {
+	TSharedPtr<SSidebarContainer> ContainerWidget = ContainerWidgetWeak.Pin();
+
 	if (InDrawer->bIsOpen)
 	{
 		// When clicking on the button of an active (but unpinned) tab, close that tab drawer
 		if (!IsDrawerPinned(InDrawer->GetUniqueId()))
 		{
-			ContainerWidget->CloseDrawer_Internal(InDrawer);
+			if (ensure(ContainerWidget.IsValid()))
+			{
+				ContainerWidget->CloseDrawer_Internal(InDrawer);
+			}
 		}
 		else if (!InDrawer->DrawerWidget->HasKeyboardFocus())
 		{
 			FSlateApplication::Get().SetKeyboardFocus(InDrawer->DrawerWidget);
 		}
 	}
-	else if (!InDrawer->State.bIsDocked)
+	else if (!InDrawer->State.bIsDocked && ensure(ContainerWidget.IsValid()))
 	{
 		// Otherwise clicking on an inactive tab should open the drawer
 		ContainerWidget->OpenDrawer_Internal(InDrawer);
@@ -351,7 +371,8 @@ void SSidebar::BuildOptionsMenu(UToolMenu* const InMenu)
 
 void SSidebar::RemoveDrawer(const TSharedRef<FSidebarDrawer>& InDrawer)
 {
-	if (InDrawer->DrawerWidget.IsValid())
+	TSharedPtr<SSidebarContainer> ContainerWidget = ContainerWidgetWeak.Pin();
+	if (ensure(ContainerWidget.IsValid()) && InDrawer->DrawerWidget.IsValid())
 	{
 		ContainerWidget->RemoveDrawerOverlaySlot(InDrawer, false);
 	}
@@ -360,7 +381,10 @@ void SSidebar::RemoveDrawer(const TSharedRef<FSidebarDrawer>& InDrawer)
 
 	Drawers.Remove(InDrawer);
 
-	ContainerWidget->UpdateDrawerTabAppearance();
+	if (ContainerWidget.IsValid())
+	{
+		ContainerWidget->UpdateDrawerTabAppearance();
+	}
 }
 
 void SSidebar::RemoveAllDrawers()
@@ -409,7 +433,12 @@ bool SSidebar::IsDrawerOpened(const FName InDrawerId) const
 
 FName SSidebar::GetOpenedDrawerId() const
 {
-	return ContainerWidget->GetOpenedDrawerId();
+	TSharedPtr<SSidebarContainer> ContainerWidget = ContainerWidgetWeak.Pin();
+	if (ensure(ContainerWidget.IsValid()))
+	{
+		return ContainerWidget->GetOpenedDrawerId();
+	}
+	return NAME_None;
 }
 
 bool SSidebar::HasDrawerPinned() const
@@ -466,7 +495,11 @@ void SSidebar::SetDrawerPinned(const FName InDrawerId, const bool bInIsPinned)
 
 		if (!DrawerToPin->bIsOpen)
 		{
-			ContainerWidget->OpenDrawerNextFrame(DrawerToPin.ToSharedRef(), false);
+			TSharedPtr<SSidebarContainer> ContainerWidget = ContainerWidgetWeak.Pin();
+			if (ensure(ContainerWidget.IsValid()))
+			{
+				ContainerWidget->OpenDrawerNextFrame(DrawerToPin.ToSharedRef(), false);
+			}
 		}
 
 		// In case two modules attempt to register drawers with initially pinned states
@@ -537,7 +570,11 @@ void SSidebar::SetDrawerDocked(const FName InDrawerId, const bool bInIsDocked)
 			SetDrawerPinned(InDrawerId, false);
 		}
 
-		ContainerWidget->DockDrawer_Internal(DrawerToDock.ToSharedRef());
+		TSharedPtr<SSidebarContainer> ContainerWidget = ContainerWidgetWeak.Pin();
+		if (ensure(ContainerWidget.IsValid()))
+		{
+			ContainerWidget->DockDrawer_Internal(DrawerToDock.ToSharedRef());
+		}
 
 		if (DrawerToDock->ContentWidget.IsValid())
 		{
@@ -556,7 +593,11 @@ void SSidebar::SetDrawerDocked(const FName InDrawerId, const bool bInIsDocked)
 	}
 	else
 	{
-		ContainerWidget->UndockDrawer_Internal(DrawerToDock.ToSharedRef());
+		TSharedPtr<SSidebarContainer> ContainerWidget = ContainerWidgetWeak.Pin();
+		if (ensure(ContainerWidget.IsValid()))
+		{
+			ContainerWidget->UndockDrawer_Internal(DrawerToDock.ToSharedRef());
+		}
 
 		if (bHideWhenAllDocked && !AreAllDrawersDocked())
         {
@@ -618,8 +659,12 @@ FSidebarState SSidebar::GetState() const
 
 	OutState.SetHidden(false);
 
-	const float CurrentDrawerSize = ContainerWidget->GetCurrentDrawerSize();
-	OutState.SetDrawerSizes(CurrentDrawerSize, 1.f - CurrentDrawerSize);
+	TSharedPtr<SSidebarContainer> ContainerWidget = ContainerWidgetWeak.Pin();
+	if (ensure(ContainerWidget.IsValid()))
+	{
+		const float CurrentDrawerSize = ContainerWidget->GetCurrentDrawerSize();
+		OutState.SetDrawerSizes(CurrentDrawerSize, 1.f - CurrentDrawerSize);
+	}
 
 	for (const TSharedRef<FSidebarDrawer>& Drawer : Drawers)
 	{
@@ -673,6 +718,12 @@ void SSidebar::SetWidgetDrawerSize(const TSharedRef<FSidebarDrawer>& InDrawer)
 		return;
 	}
 
+	TSharedPtr<SSidebarContainer> ContainerWidget = ContainerWidgetWeak.Pin();
+	if (!ensure(ContainerWidget.IsValid()))
+	{
+		return;
+	}
+
 	const float DrawerOverlayWidth = ContainerWidget->GetOverlaySize().X;
 	const float CurrentDrawerSize = ContainerWidget->GetCurrentDrawerSize();
 	const float PixelWidth = CurrentDrawerSize * DrawerOverlayWidth;
diff --git a/Engine/Source/Developer/ToolWidgets/Public/Sidebar/SSidebar.h b/Engine/Source/Developer/ToolWidgets/Public/Sidebar/SSidebar.h
--- a/Engine/Source/Developer/ToolWidgets/Public/Sidebar/SSidebar.h
+++ b/Engine/Source/Developer/ToolWidgets/Public/Sidebar/SSidebar.h
@@ -263,7 +263,7 @@ private:
 
 	const TArray<TSharedRef<FSidebarDrawer>>& GetAllDrawers() const;
 
-	TSharedPtr<SSidebarContainer> ContainerWidget;
+	TWeakPtr<SSidebarContainer> ContainerWidgetWeak;
 
 	ESidebarTabLocation TabLocation = ESidebarTabLocation::Right;
 	FOnGetContent OnGetContent;
