[Slate] Text Widgets: Add GetSelection() and SelectText() to SEditableText, SSearchBox, and SFilterSearchBox to get/set caret information. For use with Sequencer suggestion box.

diff --git a/Engine/Source/Developer/ToolWidgets/Private/Filters/SFilterSearchBox.cpp b/Engine/Source/Developer/ToolWidgets/Private/Filters/SFilterSearchBox.cpp
--- a/Engine/Source/Developer/ToolWidgets/Private/Filters/SFilterSearchBox.cpp
+++ b/Engine/Source/Developer/ToolWidgets/Private/Filters/SFilterSearchBox.cpp
@@ -1,6 +1,7 @@
 // Copyright Epic Games, Inc. All Rights Reserved.
 
 #include "Filters/SFilterSearchBox.h"
+#include "Framework/Text/TextLayout.h"
 #include "Widgets/Input/SSearchBox.h"
 #include "Widgets/Input/SMenuAnchor.h"
 #include "Widgets/Views/SListView.h"
@@ -380,4 +381,14 @@ EVisibility SFilterSearchBox::GetSearchHistoryVisibility() const
 	return bShowSearchHistory.Get() ? EVisibility::Visible : EVisibility::Collapsed;
 }
 
+FTextSelection SFilterSearchBox::GetSelection() const
+{
+	return SearchBox->GetSelection();
+}
+
+void SFilterSearchBox::SelectText(const FTextLocation& InSelectionStart, const FTextLocation& InCursorLocation)
+{
+	SearchBox->SelectText(InSelectionStart, InCursorLocation);
+}
+
 #undef LOCTEXT_NAMESPACE
diff --git a/Engine/Source/Developer/ToolWidgets/Public/Filters/SFilterSearchBox.h b/Engine/Source/Developer/ToolWidgets/Public/Filters/SFilterSearchBox.h
--- a/Engine/Source/Developer/ToolWidgets/Public/Filters/SFilterSearchBox.h
+++ b/Engine/Source/Developer/ToolWidgets/Public/Filters/SFilterSearchBox.h
@@ -10,6 +10,8 @@ class SMenuAnchor;
 class SSearchBox;
 template <typename T> class SListView;
 class SFilterSearchBoxImpl;
+class FTextSelection;
+struct FTextLocation;
 
 /** A SearchBox widget that contains support for displaying search history
  * When used with a FilterBar widget (@see SBasicFilterBar or SFilterBar), it allows the user to save searches from the
@@ -88,6 +90,12 @@ public:
 	/** Show a + button next to the current search, and set the handler for when that is clicked */
 	void SetOnSaveSearchHandler(FOnSaveSearchClicked InOnSaveSearchHandler);
 
+	/** Get the current selection */
+	FTextSelection GetSelection() const;
+
+	/** Select a block of text */
+	void SelectText(const FTextLocation& InSelectionStart, const FTextLocation& InCursorLocation);
+
 private:
 	
 	/** Handler for when text in the editable text box changed */
diff --git a/Engine/Source/Runtime/Slate/Private/Widgets/Input/SEditableText.cpp b/Engine/Source/Runtime/Slate/Private/Widgets/Input/SEditableText.cpp
--- a/Engine/Source/Runtime/Slate/Private/Widgets/Input/SEditableText.cpp
+++ b/Engine/Source/Runtime/Slate/Private/Widgets/Input/SEditableText.cpp
@@ -459,6 +459,16 @@ void SEditableText::EnableTextInputMethodContext()
 	EditableTextLayout->EnableTextInputMethodContext();
 }
 
+FTextSelection SEditableText::GetSelection() const
+{
+	return EditableTextLayout->GetSelection();
+}
+
+void SEditableText::SelectText(const FTextLocation& InSelectionStart, const FTextLocation& InCursorLocation)
+{
+	EditableTextLayout->SelectText(InSelectionStart, InCursorLocation);
+}
+
 void SEditableText::SynchronizeTextStyle()
 {
 	// Has the style used for this editable text changed?
diff --git a/Engine/Source/Runtime/Slate/Private/Widgets/Input/SSearchBox.cpp b/Engine/Source/Runtime/Slate/Private/Widgets/Input/SSearchBox.cpp
--- a/Engine/Source/Runtime/Slate/Private/Widgets/Input/SSearchBox.cpp
+++ b/Engine/Source/Runtime/Slate/Private/Widgets/Input/SSearchBox.cpp
@@ -298,3 +298,13 @@ FSlateFontInfo SSearchBox::GetWidgetFont() const
 {
 	return EditableText->GetText().IsEmpty() ? InactiveFont : ActiveFont;
 }
+
+FTextSelection SSearchBox::GetSelection() const
+{
+	return EditableText->GetSelection();
+}
+
+void SSearchBox::SelectText(const FTextLocation& InSelectionStart, const FTextLocation& InCursorLocation)
+{
+	EditableText->SelectText(InSelectionStart, InCursorLocation);
+}
diff --git a/Engine/Source/Runtime/Slate/Public/Widgets/Input/SEditableText.h b/Engine/Source/Runtime/Slate/Public/Widgets/Input/SEditableText.h
--- a/Engine/Source/Runtime/Slate/Public/Widgets/Input/SEditableText.h
+++ b/Engine/Source/Runtime/Slate/Public/Widgets/Input/SEditableText.h
@@ -369,6 +369,12 @@ public:
 	/** Register and activate the IME context for this text layout */
 	SLATE_API void EnableTextInputMethodContext();
 
+	/** Get the current selection */
+	SLATE_API FTextSelection GetSelection() const;
+
+	/** Select a block of text */
+	SLATE_API void SelectText(const FTextLocation& InSelectionStart, const FTextLocation& InCursorLocation);
+
 protected:
 	//~ Begin SWidget Interface
 	SLATE_API virtual void Tick( const FGeometry& AllottedGeometry, const double InCurrentTime, const float InDeltaTime ) override;
diff --git a/Engine/Source/Runtime/Slate/Public/Widgets/Input/SSearchBox.h b/Engine/Source/Runtime/Slate/Public/Widgets/Input/SSearchBox.h
--- a/Engine/Source/Runtime/Slate/Public/Widgets/Input/SSearchBox.h
+++ b/Engine/Source/Runtime/Slate/Public/Widgets/Input/SSearchBox.h
@@ -102,6 +102,12 @@ public:
 
 	SLATE_API void Construct( const FArguments& InArgs );
 
+	/** Get the current selection */
+	SLATE_API FTextSelection GetSelection() const;
+
+	/** Select a block of text */
+	SLATE_API void SelectText(const FTextLocation& InSelectionStart, const FTextLocation& InCursorLocation);
+
 private:
 
 	/** Callback for changes in the editable text box. */
