[BUGFIX]Action Button Incorrect Padding
- Action buttons now treat the nullptr returned from a attribute delegate invalid and set the padding accordingly.

diff --git a/Engine/Source/Developer/ToolWidgets/Private/ToolWidgetsUtilitiesPrivate.cpp b/Engine/Source/Developer/ToolWidgets/Private/ToolWidgetsUtilitiesPrivate.cpp
--- a/Engine/Source/Developer/ToolWidgets/Private/ToolWidgetsUtilitiesPrivate.cpp
+++ b/Engine/Source/Developer/ToolWidgets/Private/ToolWidgetsUtilitiesPrivate.cpp
@@ -26,11 +26,22 @@ namespace UE::ToolWidgets::Private
 			check(InIconColorAndOpacity.IsSet() || InIconColorAndOpacity.IsBound());
 			check(InTextBlockStyle);
 
-			const bool bHasIcon = InIcon.Get() || InIcon.IsBound();
+			const bool bCanEverHaveIcon = InIcon.Get() || InIcon.IsBound();
 
 			static constexpr float DefaultIconHeight = FToolWidgetsStylePrivate::FActionButton::DefaultIconHeight;
 			static constexpr float IconTextPadding = FToolWidgetsStylePrivate::FActionButton::DefaultIconLabelSpacing;
 
+
+			TAttribute<FMargin> PaddingAttribute = TAttribute<FMargin>::CreateLambda([InIcon, DefaultPadding = IconTextPadding]()
+			{
+				if (InIcon.Get(nullptr) != nullptr)
+				{
+					return FMargin(DefaultPadding, 0, 0, 0);
+				}
+
+				return FMargin(0);
+			});
+
 			const TSharedRef<SHorizontalBox> ButtonContentContainer = SNew(SHorizontalBox)
 				+ SHorizontalBox::Slot()
 				.HAlign(HAlign_Center)
@@ -39,7 +50,7 @@ namespace UE::ToolWidgets::Private
 				.Padding(0)
 				[
 					SNew(SWidgetSwitcher)
-					.WidgetIndex(bHasIcon ? 1 : 0)
+					.WidgetIndex(bCanEverHaveIcon ? 1 : 0)
 
 					+ SWidgetSwitcher::Slot()
 					[
@@ -52,14 +63,14 @@ namespace UE::ToolWidgets::Private
 						SNew(SImage)
 						.Image(InIcon)
 						.ColorAndOpacity(InIconColorAndOpacity)
-						.Visibility(bHasIcon ? EVisibility::HitTestInvisible : EVisibility::Collapsed)
+						.Visibility(bCanEverHaveIcon ? EVisibility::HitTestInvisible : EVisibility::Collapsed)
 					]
 				]
 
 				+ SHorizontalBox::Slot()
 				.VAlign(VAlign_Center)
 				.FillWidth(1.0f)
-				.Padding(FMargin(bHasIcon ? IconTextPadding : 0, 0, 0, 0))
+				.Padding(PaddingAttribute)
 				[
 					SNew(STextBlock)
 					.TextStyle(InTextBlockStyle)
