Move SVerticalResizeBox from Niagara editor to ToolWidgets module for re-use in other editors.

diff --git a/Engine/Plugins/FX/Niagara/Source/NiagaraEditor/Private/Customizations/NiagaraOutlinerCustomization.cpp b/Engine/Plugins/FX/Niagara/Source/NiagaraEditor/Private/Customizations/NiagaraOutlinerCustomization.cpp
--- a/Engine/Plugins/FX/Niagara/Source/NiagaraEditor/Private/Customizations/NiagaraOutlinerCustomization.cpp
+++ b/Engine/Plugins/FX/Niagara/Source/NiagaraEditor/Private/Customizations/NiagaraOutlinerCustomization.cpp
@@ -24,11 +24,11 @@
 #include "Widgets/Layout/SSplitter.h"
 #include "Widgets/Images/SImage.h"
 #include "Widgets/Input/SButton.h"
+#include "SResizeBox.h"
 ///Niagara
 #include "NiagaraEditorModule.h"
 #include "NiagaraComponent.h"
 #include "NiagaraEditorStyle.h"
-#include "Widgets/SVerticalResizeBox.h"
 
 #include "NiagaraSimCache.h"
 
diff --git a/Engine/Plugins/FX/Niagara/Source/NiagaraEditor/Private/Widgets/SVerticalResizeBox.cpp b/Engine/Plugins/FX/Niagara/Source/NiagaraEditor/Private/Widgets/SVerticalResizeBox.cpp
deleted file mode 100644
diff --git a/Engine/Plugins/FX/Niagara/Source/NiagaraEditor/Public/Widgets/SVerticalResizeBox.h b/Engine/Plugins/FX/Niagara/Source/NiagaraEditor/Public/Widgets/SVerticalResizeBox.h
--- a/Engine/Plugins/FX/Niagara/Source/NiagaraEditor/Public/Widgets/SVerticalResizeBox.h
+++ b/Engine/Plugins/FX/Niagara/Source/NiagaraEditor/Public/Widgets/SVerticalResizeBox.h
@@ -2,54 +2,10 @@
 
 #pragma once
 
-#include "Widgets/SCompoundWidget.h"
-#include "Widgets/DeclarativeSyntaxSupport.h"
-#include "Styling/SlateBrush.h"
-#include "Types/SlateStructs.h"
+// HEADER_UNIT_SKIP - Deprecated
 
-class SVerticalResizeBox : public SCompoundWidget
-{
-public:
-	DECLARE_DELEGATE_OneParam(FOnContentHeightChanged, float);
+#include "CoreTypes.h"
 
-public:
-	SLATE_BEGIN_ARGS(SVerticalResizeBox)
-		: _HandleHeight(5)
-		, _ContentHeight(50)
-		, _HandleColor(FLinearColor(0.0f, 0.0f, 0.0f, 0.0f))
-		, _HandleHighlightColor(FLinearColor(1.0f, 1.0f, 1.0f, 0.5f))
-	{}
-		SLATE_ARGUMENT(float, HandleHeight)
-		SLATE_ATTRIBUTE(float, ContentHeight)
-		SLATE_ATTRIBUTE(FLinearColor, HandleColor)
-		SLATE_ATTRIBUTE(FLinearColor, HandleHighlightColor)
-		SLATE_EVENT(FOnContentHeightChanged, ContentHeightChanged)
-		SLATE_DEFAULT_SLOT(FArguments, Content);
-	SLATE_END_ARGS();
+UE_DEPRECATED_HEADER(5.6, "Use SResizeBox.h instead of Widgets/SVerticalResizeBox.h, and add ToolWidgets as a module dependency")
 
-	NIAGARAEDITOR_API void Construct(const FArguments& InArgs);
-
-	NIAGARAEDITOR_API virtual FReply OnMouseButtonDown(const FGeometry& MyGeometry, const FPointerEvent& MouseEvent) override;
-	NIAGARAEDITOR_API virtual FReply OnMouseButtonUp(const FGeometry& MyGeometry, const FPointerEvent& MouseEvent) override;
-	NIAGARAEDITOR_API virtual FReply OnMouseMove(const FGeometry& MyGeometry, const FPointerEvent& MouseEvent) override;
-	NIAGARAEDITOR_API virtual FCursorReply OnCursorQuery(const FGeometry& MyGeometry, const FPointerEvent& CursorEvent) const override;
-	NIAGARAEDITOR_API virtual int32 OnPaint(const FPaintArgs& Args, const FGeometry& AllottedGeometry, const FSlateRect& MyClippingRect, FSlateWindowElementList& OutDrawElements, int32 LayerId, const FWidgetStyle& InWidgetStyle, bool bParentEnabled) const override;
-
-private:
-	FOptionalSize GetHeightOverride() const;
-
-private:
-	TOptional<float> LastMouseLocation;
-
-	TAttribute<float> ContentHeight;
-	float HandleHeight;
-
-	float DragStartLocation;
-	float DragStartContentHeight;
-
-	TAttribute<FLinearColor> HandleColor;
-	TAttribute<FLinearColor> HandleHighlightColor;
-	FSlateBrush HandleBrush;
-
-	FOnContentHeightChanged ContentHeightChanged;
-};
+#include "SResizeBox.h"
diff --git a/Engine/Plugins/FX/Niagara/Source/NiagaraEditorWidgets/Private/DetailCustomizations/NiagaraDataInterfaceCurveDetails.cpp b/Engine/Plugins/FX/Niagara/Source/NiagaraEditorWidgets/Private/DetailCustomizations/NiagaraDataInterfaceCurveDetails.cpp
--- a/Engine/Plugins/FX/Niagara/Source/NiagaraEditorWidgets/Private/DetailCustomizations/NiagaraDataInterfaceCurveDetails.cpp
+++ b/Engine/Plugins/FX/Niagara/Source/NiagaraEditorWidgets/Private/DetailCustomizations/NiagaraDataInterfaceCurveDetails.cpp
@@ -31,7 +31,7 @@
 #include "SCurveKeyDetailPanel.h"
 #include "CurveEditorCommands.h"
 
-#include "Widgets/SVerticalResizeBox.h"
+#include "SResizeBox.h"
 #include "Widgets/Input/SComboButton.h"
 #include "Widgets/Input/SButton.h"
 #include "Widgets/Images/SImage.h"
diff --git a/Engine/Plugins/FX/Niagara/Source/NiagaraEditorWidgets/Private/SNiagaraScratchPadScriptManager.cpp b/Engine/Plugins/FX/Niagara/Source/NiagaraEditorWidgets/Private/SNiagaraScratchPadScriptManager.cpp
--- a/Engine/Plugins/FX/Niagara/Source/NiagaraEditorWidgets/Private/SNiagaraScratchPadScriptManager.cpp
+++ b/Engine/Plugins/FX/Niagara/Source/NiagaraEditorWidgets/Private/SNiagaraScratchPadScriptManager.cpp
@@ -6,7 +6,6 @@
 #include "Widgets/SDynamicLayoutBox.h"
 #include "Widgets/SItemSelector.h"
 #include "Widgets/SNiagaraSelectedObjectsDetails.h"
-#include "Widgets/SVerticalResizeBox.h"
 #include "Widgets/SNiagaraScriptGraph.h"
 #include "NiagaraObjectSelection.h"
 #include "NiagaraNodeFunctionCall.h"
@@ -18,6 +17,7 @@
 #include "Widgets/Text/SInlineEditableTextBlock.h"
 #include "Widgets/Input/SButton.h"
 #include "Widgets/Layout/SSpacer.h"
+#include "SResizeBox.h"
 #include "EditorStyleSet.h"
 #include "Framework/MultiBox/MultiBoxBuilder.h"
 #include "Framework/Application/SlateApplication.h"
diff --git a/Engine/Source/Developer/ToolWidgets/Private/SResizeBox.cpp b/Engine/Source/Developer/ToolWidgets/Private/SResizeBox.cpp
new file mode 100644
--- /dev/null
+++ b/Engine/Source/Developer/ToolWidgets/Private/SResizeBox.cpp
@@ -0,0 +1,113 @@
+// Copyright Epic Games, Inc. All Rights Reserved.
+
+#include "SResizeBox.h"
+#include "Brushes/SlateColorBrush.h"
+#include "Widgets/Layout/SBox.h"
+
+void SVerticalResizeBox::Construct(const FArguments& InArgs)
+{
+	ContentHeight = InArgs._ContentHeight;
+	HandleHeight = InArgs._HandleHeight;
+	HandleColor = InArgs._HandleColor;
+	HandleHighlightColor = InArgs._HandleHighlightColor;
+	HandleBrush = FSlateColorBrush(FLinearColor::White);
+	ContentHeightChanged = InArgs._ContentHeightChanged;
+
+	ChildSlot
+	[
+		SNew(SBox)
+		.HeightOverride(this, &SVerticalResizeBox::GetHeightOverride)
+		.Padding(FMargin(0, 0, 0, HandleHeight))
+		[
+			InArgs._Content.Widget
+		]
+	];
+}
+
+FReply SVerticalResizeBox::OnMouseButtonDown(const FGeometry& MyGeometry, const FPointerEvent& MouseEvent)
+{
+	if (MouseEvent.GetEffectingButton() == EKeys::LeftMouseButton)
+	{
+		FVector2D MouseLocation = MyGeometry.AbsoluteToLocal(MouseEvent.GetScreenSpacePosition());
+		if (MyGeometry.GetLocalSize().Y - MouseLocation.Y < HandleHeight)
+		{
+			DragStartLocation = MouseLocation.Y;
+			DragStartContentHeight = ContentHeight.Get();
+			return FReply::Handled().CaptureMouse(SharedThis(this));
+		}
+	}
+	return FReply::Unhandled();
+}
+
+FReply SVerticalResizeBox::OnMouseButtonUp(const FGeometry& MyGeometry, const FPointerEvent& MouseEvent)
+{
+	if (this->HasMouseCapture())
+	{
+		return FReply::Handled().ReleaseMouseCapture();
+	}
+	return FReply::Unhandled();
+}
+
+FReply SVerticalResizeBox::OnMouseMove(const FGeometry& MyGeometry, const FPointerEvent& MouseEvent)
+{
+	FVector2D MouseLocation = MyGeometry.AbsoluteToLocal(MouseEvent.GetScreenSpacePosition());
+	LastMouseLocation = MouseLocation.Y;
+	if (this->HasMouseCapture())
+	{
+		float NewContentHeight = DragStartContentHeight + (MouseLocation.Y - DragStartLocation);
+		if (ContentHeight.IsBound() && ContentHeightChanged.IsBound())
+		{
+			ContentHeightChanged.Execute(NewContentHeight);
+		}
+		else
+		{
+			ContentHeight = NewContentHeight;
+		}
+		return FReply::Handled();
+	}
+	return FReply::Unhandled();
+}
+
+FCursorReply SVerticalResizeBox::OnCursorQuery(const FGeometry& MyGeometry, const FPointerEvent& CursorEvent) const
+{
+	FVector2D CursorLocation = MyGeometry.AbsoluteToLocal(CursorEvent.GetScreenSpacePosition());
+	if (MyGeometry.GetLocalSize().Y - CursorLocation.Y < HandleHeight)
+	{
+		return FCursorReply::Cursor(EMouseCursor::ResizeUpDown);
+	}
+	return FCursorReply::Unhandled();
+}
+
+int32 SVerticalResizeBox::OnPaint(const FPaintArgs& Args, const FGeometry& AllottedGeometry, const FSlateRect& MyClippingRect, FSlateWindowElementList& OutDrawElements, int32 LayerId, const FWidgetStyle& InWidgetStyle, bool bParentEnabled) const
+{
+	FLinearColor HandleBoxColor;
+	int32 HandleLayerId = LayerId + 1;
+	FVector2D LocalSize = AllottedGeometry.GetLocalSize();
+	if (IsHovered() && LastMouseLocation.IsSet() && LastMouseLocation.GetValue() >= LocalSize.Y - HandleHeight && LastMouseLocation.GetValue() <= LocalSize.Y)
+	{
+		HandleBoxColor = HandleHighlightColor.Get();
+	}
+	else
+	{
+		HandleBoxColor = HandleColor.Get();
+	}
+
+	FVector2D HandleLocation(0, AllottedGeometry.GetLocalSize().Y - HandleHeight);
+	FVector2D HandleSize(AllottedGeometry.GetLocalSize().X, HandleHeight);
+	FSlateDrawElement::MakeBox
+	(
+		OutDrawElements,
+		HandleLayerId,
+		AllottedGeometry.ToPaintGeometry(HandleSize, FSlateLayoutTransform(HandleLocation)),
+		&HandleBrush,
+		ESlateDrawEffect::None,
+		HandleBoxColor
+	);
+
+	return SCompoundWidget::OnPaint(Args, AllottedGeometry, MyClippingRect, OutDrawElements, HandleLayerId, InWidgetStyle, bParentEnabled);
+}
+
+FOptionalSize SVerticalResizeBox::GetHeightOverride() const
+{
+	return ContentHeight.Get() + HandleHeight;
+}
\ No newline at end of file
diff --git a/Engine/Source/Developer/ToolWidgets/Public/SResizeBox.h b/Engine/Source/Developer/ToolWidgets/Public/SResizeBox.h
new file mode 100644
--- /dev/null
+++ b/Engine/Source/Developer/ToolWidgets/Public/SResizeBox.h
@@ -0,0 +1,55 @@
+// Copyright Epic Games, Inc. All Rights Reserved.
+
+#pragma once
+
+#include "Widgets/SCompoundWidget.h"
+#include "Widgets/DeclarativeSyntaxSupport.h"
+#include "Styling/SlateBrush.h"
+#include "Types/SlateStructs.h"
+
+class SVerticalResizeBox : public SCompoundWidget
+{
+public:
+	DECLARE_DELEGATE_OneParam(FOnContentHeightChanged, float);
+
+public:
+	SLATE_BEGIN_ARGS(SVerticalResizeBox)
+		: _HandleHeight(5)
+		, _ContentHeight(50)
+		, _HandleColor(FLinearColor(0.0f, 0.0f, 0.0f, 0.0f))
+		, _HandleHighlightColor(FLinearColor(1.0f, 1.0f, 1.0f, 0.5f))
+	{}
+		SLATE_ARGUMENT(float, HandleHeight)
+		SLATE_ATTRIBUTE(float, ContentHeight)
+		SLATE_ATTRIBUTE(FLinearColor, HandleColor)
+		SLATE_ATTRIBUTE(FLinearColor, HandleHighlightColor)
+		SLATE_EVENT(FOnContentHeightChanged, ContentHeightChanged)
+		SLATE_DEFAULT_SLOT(FArguments, Content);
+	SLATE_END_ARGS();
+
+	TOOLWIDGETS_API void Construct(const FArguments& InArgs);
+
+	TOOLWIDGETS_API virtual FReply OnMouseButtonDown(const FGeometry& MyGeometry, const FPointerEvent& MouseEvent) override;
+	TOOLWIDGETS_API virtual FReply OnMouseButtonUp(const FGeometry& MyGeometry, const FPointerEvent& MouseEvent) override;
+	TOOLWIDGETS_API virtual FReply OnMouseMove(const FGeometry& MyGeometry, const FPointerEvent& MouseEvent) override;
+	TOOLWIDGETS_API virtual FCursorReply OnCursorQuery(const FGeometry& MyGeometry, const FPointerEvent& CursorEvent) const override;
+	TOOLWIDGETS_API virtual int32 OnPaint(const FPaintArgs& Args, const FGeometry& AllottedGeometry, const FSlateRect& MyClippingRect, FSlateWindowElementList& OutDrawElements, int32 LayerId, const FWidgetStyle& InWidgetStyle, bool bParentEnabled) const override;
+
+private:
+	FOptionalSize GetHeightOverride() const;
+
+private:
+	TOptional<float> LastMouseLocation;
+
+	TAttribute<float> ContentHeight;
+	float HandleHeight;
+
+	float DragStartLocation;
+	float DragStartContentHeight;
+
+	TAttribute<FLinearColor> HandleColor;
+	TAttribute<FLinearColor> HandleHighlightColor;
+	FSlateBrush HandleBrush;
+
+	FOnContentHeightChanged ContentHeightChanged;
+};
