[ToolWidgets] Sidebar: Remove scroll box widget wrapper around drawer content to give devs more flexibility. Devs will need to wrap the drawer content manually if they want a scroll bar.

diff --git a/Engine/Source/Developer/ToolWidgets/Private/Sidebar/SSidebarDrawerContent.cpp b/Engine/Source/Developer/ToolWidgets/Private/Sidebar/SSidebarDrawerContent.cpp
--- a/Engine/Source/Developer/ToolWidgets/Private/Sidebar/SSidebarDrawerContent.cpp
+++ b/Engine/Source/Developer/ToolWidgets/Private/Sidebar/SSidebarDrawerContent.cpp
@@ -6,7 +6,6 @@
 #include "Sidebar/SidebarDrawer.h"
 #include "Widgets/Input/SCheckBox.h"
 #include "Widgets/Layout/SBox.h"
-#include "Widgets/Layout/SScrollBox.h"
 #include "Widgets/Layout/SWrapBox.h"
 #include "Widgets/SBoxPanel.h"
 #include "Widgets/Text/STextBlock.h"
@@ -43,8 +42,7 @@ void SSidebarDrawerContent::Construct(const FArguments& InArgs, const TWeakPtr<F
 		+ SVerticalBox::Slot()
 		.FillHeight(1.f)
 		[
-			SAssignNew(ContentBox, SScrollBox)
-			.Orientation(Orient_Vertical)
+			SAssignNew(ContentBox, SVerticalBox)
 		]
 	];
 
@@ -262,23 +260,19 @@ void SSidebarDrawerContent::AddContentSlot(const TSharedRef<ISidebarDrawerConten
 			InDrawerContent->CreateContentWidget()
 		];
 
+	SVerticalBox::FSlot::FSlotArguments& Slot = ContentBox->AddSlot()
+		.Padding(0.f, 0.f, 0.f, BottomPadding)
+		[
+			ContentWidget
+		];
+
 	if (InFillSize.IsSet())
 	{
-		ContentBox->AddSlot()
-			.FillSize(InFillSize.Get(1.f))
-			.Padding(0.f, 0.f, 0.f, BottomPadding)
-			[
-				ContentWidget
-			];
+		Slot.FillHeight(InFillSize.GetValue());
 	}
 	else
 	{
-		ContentBox->AddSlot()
-			.AutoSize()
-			.Padding(0.f, 0.f, 0.f, BottomPadding)
-			[
-				ContentWidget
-			];
+		Slot.AutoHeight();
 	}
 }
 
diff --git a/Engine/Source/Developer/ToolWidgets/Private/Sidebar/SSidebarDrawerContent.h b/Engine/Source/Developer/ToolWidgets/Private/Sidebar/SSidebarDrawerContent.h
--- a/Engine/Source/Developer/ToolWidgets/Private/Sidebar/SSidebarDrawerContent.h
+++ b/Engine/Source/Developer/ToolWidgets/Private/Sidebar/SSidebarDrawerContent.h
@@ -10,7 +10,7 @@
 
 class FSidebarDrawer;
 class ISidebarDrawerContent;
-class SScrollBox;
+class SVerticalBox;
 class SWrapBox;
 
 /** Handles drawer multi-section display */
@@ -44,5 +44,6 @@ private:
 	TWeakPtr<FSidebarDrawer> OwnerDrawerWeak;
 
 	TSharedPtr<SWrapBox> ButtonBox;
-	TSharedPtr<SScrollBox> ContentBox;
+
+	TSharedPtr<SVerticalBox> ContentBox;
 };
diff --git a/Engine/Source/Editor/Sequencer/Private/Widgets/Sidebar/SequencerSelectionDrawer.cpp b/Engine/Source/Editor/Sequencer/Private/Widgets/Sidebar/SequencerSelectionDrawer.cpp
--- a/Engine/Source/Editor/Sequencer/Private/Widgets/Sidebar/SequencerSelectionDrawer.cpp
+++ b/Engine/Source/Editor/Sequencer/Private/Widgets/Sidebar/SequencerSelectionDrawer.cpp
@@ -25,6 +25,7 @@
 #include "Styling/AppStyle.h"
 #include "Templates/SharedPointer.h"
 #include "Widgets/Layout/SBorder.h"
+#include "Widgets/Layout/SScrollBox.h"
 #include "Widgets/Sidebar/SMarkedFrameDetails.h"
 #include "Widgets/Sidebar/STrackDetails.h"
 
@@ -157,15 +158,19 @@ TSharedRef<SWidget> FSequencerSelectionDrawer::CreateContentWidget()
 		}
 	}
 
-	return SNew(SBorder)
-		.BorderImage(FAppStyle::GetBrush(TEXT("NoBorder")))
-		.Padding(0.f)
+	return SNew(SScrollBox)
+		+ SScrollBox::Slot()
 		[
-			SAssignNew(ContentBox, SVerticalBox)
-			+ SVerticalBox::Slot()
-			.FillHeight(1.f)
+			SNew(SBorder)
+			.BorderImage(FAppStyle::GetBrush(TEXT("NoBorder")))
+			.Padding(0.f)
 			[
-				CreateNoSelectionHintText()
+				SAssignNew(ContentBox, SVerticalBox)
+				+ SVerticalBox::Slot()
+				.FillHeight(1.f)
+				[
+					CreateNoSelectionHintText()
+				]
 			]
 		];
 }
