[Sequencer] Fix for engine test failing for Linux with TOptional::Get(1.f) instead of GetValue()

diff --git a/Engine/Source/Developer/ToolWidgets/Private/Sidebar/SSidebarDrawerContent.cpp b/Engine/Source/Developer/ToolWidgets/Private/Sidebar/SSidebarDrawerContent.cpp
--- a/Engine/Source/Developer/ToolWidgets/Private/Sidebar/SSidebarDrawerContent.cpp
+++ b/Engine/Source/Developer/ToolWidgets/Private/Sidebar/SSidebarDrawerContent.cpp
@@ -268,7 +268,7 @@ void SSidebarDrawerContent::AddContentSlot(const TSharedRef<ISidebarDrawerConten
 
 	if (InFillSize.IsSet())
 	{
-		Slot.FillHeight(InFillSize.GetValue());
+		Slot.FillHeight(InFillSize.Get(1.f));
 	}
 	else
 	{
