[Sequencer] UE.EditorAutomation(RunTest=Animation) Linux_ASan vulkan - Keyframe creation tests failing due to critical failure

diff --git a/Engine/Source/Developer/ToolWidgets/Private/Sidebar/SSidebarDrawerContent.cpp b/Engine/Source/Developer/ToolWidgets/Private/Sidebar/SSidebarDrawerContent.cpp
--- a/Engine/Source/Developer/ToolWidgets/Private/Sidebar/SSidebarDrawerContent.cpp
+++ b/Engine/Source/Developer/ToolWidgets/Private/Sidebar/SSidebarDrawerContent.cpp
@@ -128,7 +128,7 @@ void SSidebarDrawerContent::BuildContent()
 
 	for (const TPair<FName, TSharedRef<ISidebarDrawerContent>>& SectionPair : SortedContentDrawers)
 	{
-		AddContentSlot(SectionPair.Value, 1.f);
+		AddContentSlot(SectionPair.Value);
 	}
 
 	if (Drawer->State.SelectedSections.IsEmpty())
@@ -249,31 +249,37 @@ TArray<TSharedRef<ISidebarDrawerContent>> SSidebarDrawerContent::GetOrderedSecti
 	return MoveTemp(OrderedSections);
 }
 
-void SSidebarDrawerContent::AddContentSlot(const TSharedRef<ISidebarDrawerContent>& InDrawerContent, const TOptional<float>& InFillSize)
+void SSidebarDrawerContent::AddContentSlot(const TSharedRef<ISidebarDrawerContent>& InDrawerContent)
 {
-	const float BottomPadding = 2.f;
-	const FName SectionName = InDrawerContent->GetSectionId();
-
-	const TSharedRef<SBox> ContentWidget = SNew(SBox)
-		.Visibility(this, &SSidebarDrawerContent::GetSectionContentVisibility, SectionName, InDrawerContent->AsWeak())
-		[
-			InDrawerContent->CreateContentWidget()
-		];
-
 	SVerticalBox::FSlot::FSlotArguments& Slot = ContentBox->AddSlot()
-		.Padding(0.f, 0.f, 0.f, BottomPadding)
+		.FillHeight(1.f) // Remove this after jira UE-237053 has been resolved to allow drawers to specify FillHeight through the ISidebarDrawerContent interface
+		.Padding(0.f, 0.f, 0.f, 2.f)
 		[
-			ContentWidget
+			SNew(SBox)
+			.Visibility(this, &SSidebarDrawerContent::GetSectionContentVisibility
+				, InDrawerContent->GetSectionId(), InDrawerContent->AsWeak())
+			[
+				InDrawerContent->CreateContentWidget()
+			]
 		];
 
-	if (InFillSize.IsSet())
+	/*
+	 * Not using this until jira UE-237053 is resolved.
+	 * Currently, this is not being used anywhere, but ISidebarDrawerContent::GetSectionFillSize
+	 * should be added to allow a drawer to specify the fill size themselves.
+	 * Motion Design, for example, could use this to better display its sections in its sidebar.
+	 * There is still some discussion on if the Motion Design sidebar is going to remain being used
+	 * or if it will be integrated into the Navigation Tool to display embedded sequences.
+	 */
+	/*const TOptional<float> FillSize = InDrawerContent->GetSectionFillSize();
+	if (FillSize.IsSet())
 	{
-		Slot.FillHeight(InFillSize.Get(1.f));
+		Slot.FillHeight(FillSize.Get(1.f));
 	}
 	else
 	{
 		Slot.AutoHeight();
-	}
+	}*/
 }
 
 #undef LOCTEXT_NAMESPACE 
diff --git a/Engine/Source/Developer/ToolWidgets/Private/Sidebar/SSidebarDrawerContent.h b/Engine/Source/Developer/ToolWidgets/Private/Sidebar/SSidebarDrawerContent.h
--- a/Engine/Source/Developer/ToolWidgets/Private/Sidebar/SSidebarDrawerContent.h
+++ b/Engine/Source/Developer/ToolWidgets/Private/Sidebar/SSidebarDrawerContent.h
@@ -39,7 +39,7 @@ private:
 
 	TArray<TSharedRef<ISidebarDrawerContent>> GetOrderedSections() const;
 
-	void AddContentSlot(const TSharedRef<ISidebarDrawerContent>& InDrawerContent, const TOptional<float>& InFillSize = TOptional<float>());
+	void AddContentSlot(const TSharedRef<ISidebarDrawerContent>& InDrawerContent);
 
 	TWeakPtr<FSidebarDrawer> OwnerDrawerWeak;
 
