[ToolWidgets] Sidebar:
- Fix drawer removing code causing crash when sidebar is recreated in some cases
- Unpin and undock all other drawers when a drawer is docked

diff --git a/Engine/Source/Developer/ToolWidgets/Private/Sidebar/SSidebar.cpp b/Engine/Source/Developer/ToolWidgets/Private/Sidebar/SSidebar.cpp
--- a/Engine/Source/Developer/ToolWidgets/Private/Sidebar/SSidebar.cpp
+++ b/Engine/Source/Developer/ToolWidgets/Private/Sidebar/SSidebar.cpp
@@ -132,10 +132,7 @@ bool SSidebar::UnregisterDrawer(const FName InDrawerId)
 		return false;
 	}
 
-	TabButtonContainer->RemoveSlot(Drawers[IndexToRemove]->ButtonWidget.ToSharedRef());
-
 	RemoveDrawer(Drawers[IndexToRemove]);
-	Drawers.RemoveAt(IndexToRemove);
 
 	if (ContainerWidget.IsValid())
 	{
@@ -371,15 +368,25 @@ void SSidebar::BuildOptionsMenu(UToolMenu* const InMenu)
 
 void SSidebar::RemoveDrawer(const TSharedRef<FSidebarDrawer>& InDrawer)
 {
-	TSharedPtr<SSidebarContainer> ContainerWidget = ContainerWidgetWeak.Pin();
+	const FName DrawerId = InDrawer->GetUniqueId();
+
+	const TSharedPtr<SSidebarContainer> ContainerWidget = ContainerWidgetWeak.Pin();
 	if (ensure(ContainerWidget.IsValid()) && InDrawer->DrawerWidget.IsValid())
 	{
 		ContainerWidget->RemoveDrawerOverlaySlot(InDrawer, false);
 	}
 
-	InDrawer->DrawerClosedDelegate.ExecuteIfBound(InDrawer->GetUniqueId());
+	const int32 IndexToRemove = Drawers.IndexOfByPredicate(
+		[DrawerId](const TSharedRef<FSidebarDrawer>& InDrawer)
+		{
+			return DrawerId == InDrawer->GetUniqueId();
+		});
+	if (IndexToRemove != INDEX_NONE)
+	{
+		Drawers.RemoveAt(IndexToRemove);
+	}
 
-	Drawers.Remove(InDrawer);
+	InDrawer->DrawerClosedDelegate.ExecuteIfBound(DrawerId);
 
 	if (ContainerWidget.IsValid())
 	{
@@ -488,6 +495,8 @@ void SSidebar::SetDrawerPinned(const FName InDrawerId, const bool bInIsPinned)
 
 	if (bInIsPinned)
 	{
+		UnpinAllDrawers();
+
 		if (DrawerToPin->State.bIsDocked)
 		{
 			SetDrawerDocked(InDrawerId, false);
@@ -563,6 +572,23 @@ void SSidebar::SetDrawerDocked(const FName InDrawerId, const bool bInIsDocked)
 		return;
 	}
 
+	const TSharedPtr<SSidebarContainer> ContainerWidget = ContainerWidgetWeak.Pin();
+	if (!ensure(ContainerWidget.IsValid()))
+	{
+		return;
+	}
+
+	// Unpin and undock all over tabs
+	for (const TSharedRef<FSidebarDrawer>& Drawer : Drawers)
+	{
+		if (bInIsDocked)
+		{
+			SetDrawerPinned(Drawer->GetUniqueId(), false);
+		}
+
+		ContainerWidget->UndockDrawer_Internal(Drawer);
+	}
+
 	if (bInIsDocked)
 	{
 		if (DrawerToDock->State.bIsPinned)
@@ -570,11 +596,7 @@ void SSidebar::SetDrawerDocked(const FName InDrawerId, const bool bInIsDocked)
 			SetDrawerPinned(InDrawerId, false);
 		}
 
-		TSharedPtr<SSidebarContainer> ContainerWidget = ContainerWidgetWeak.Pin();
-		if (ensure(ContainerWidget.IsValid()))
-		{
-			ContainerWidget->DockDrawer_Internal(DrawerToDock.ToSharedRef());
-		}
+		ContainerWidget->DockDrawer_Internal(DrawerToDock.ToSharedRef());
 
 		if (DrawerToDock->ContentWidget.IsValid())
 		{
@@ -593,12 +615,6 @@ void SSidebar::SetDrawerDocked(const FName InDrawerId, const bool bInIsDocked)
 	}
 	else
 	{
-		TSharedPtr<SSidebarContainer> ContainerWidget = ContainerWidgetWeak.Pin();
-		if (ensure(ContainerWidget.IsValid()))
-		{
-			ContainerWidget->UndockDrawer_Internal(DrawerToDock.ToSharedRef());
-		}
-
 		if (bHideWhenAllDocked && !AreAllDrawersDocked())
         {
         	SetVisibility(EVisibility::Visible);
