[Truncation Warnings] ToolMenu and ToolWidget modules

diff --git a/Engine/Source/Developer/ToolMenus/Private/ToolMenu.cpp b/Engine/Source/Developer/ToolMenus/Private/ToolMenu.cpp
--- a/Engine/Source/Developer/ToolMenus/Private/ToolMenu.cpp
+++ b/Engine/Source/Developer/ToolMenus/Private/ToolMenu.cpp
@@ -24,7 +24,7 @@ UToolMenu::UToolMenu() :
 	, bRegistered(false)
 	, bIsRegistering(false)
 	, bExtendersEnabled(true)
-	, MaxHeight(1000.f)
+	, MaxHeight(1000)
 {
 }
 
diff --git a/Engine/Source/Developer/ToolMenus/ToolMenus.Build.cs b/Engine/Source/Developer/ToolMenus/ToolMenus.Build.cs
--- a/Engine/Source/Developer/ToolMenus/ToolMenus.Build.cs
+++ b/Engine/Source/Developer/ToolMenus/ToolMenus.Build.cs
@@ -6,6 +6,9 @@ namespace UnrealBuildTool.Rules
 	{
 		public ToolMenus(ReadOnlyTargetRules Target) : base(Target)
 		{
+			// Enable truncation warnings in this module
+			CppCompileWarningSettings.UnsafeTypeCastWarningLevel = WarningLevel.Warning;
+
 			PrivateDependencyModuleNames.AddRange(
 				new string[] {
 					"Core",
diff --git a/Engine/Source/Developer/ToolWidgets/Private/Dialog/SCustomDialog.cpp b/Engine/Source/Developer/ToolWidgets/Private/Dialog/SCustomDialog.cpp
--- a/Engine/Source/Developer/ToolWidgets/Private/Dialog/SCustomDialog.cpp
+++ b/Engine/Source/Developer/ToolWidgets/Private/Dialog/SCustomDialog.cpp
@@ -85,7 +85,7 @@ TSharedRef<SWidget> SCustomDialog::CreateContentBox(const FArguments& InArgs)
 		.HAlign(InArgs._HAlignIcon)
 		[
 			SNew(SBox)
-			.Padding(0, 0, 8, 0)
+			.Padding(0.f, 0.f, 8.f, 0.f)
 			.Visibility_Lambda([IconAttr = InArgs._Icon]() { return IconAttr.Get(nullptr) ? EVisibility::Visible : EVisibility::Collapsed; })
 			[
 				SNew(SImage)
diff --git a/Engine/Source/Developer/ToolWidgets/Private/Filters/SCustomTextFilterDialog.cpp b/Engine/Source/Developer/ToolWidgets/Private/Filters/SCustomTextFilterDialog.cpp
--- a/Engine/Source/Developer/ToolWidgets/Private/Filters/SCustomTextFilterDialog.cpp
+++ b/Engine/Source/Developer/ToolWidgets/Private/Filters/SCustomTextFilterDialog.cpp
@@ -28,8 +28,8 @@ void SCustomTextFilterDialog::Construct( const FArguments& InArgs )
 	{
 	    // Button to modify the current filter
        	ButtonBox->AddSlot()
-		.FillWidth(1.0)
-       	.Padding(0, 0, 16, 0)
+		.FillWidth(1.0f)
+       	.Padding(0.f, 0.f, 16.f, 0.f)
        	.HAlign(HAlign_Right)
        	[
        		SNew(SPrimaryButton)
@@ -40,7 +40,7 @@ void SCustomTextFilterDialog::Construct( const FArguments& InArgs )
 		// Button to delete the current filter
 		ButtonBox->AddSlot()
 		.AutoWidth()
-		.Padding(0, 0, 16, 0)
+		.Padding(0.f, 0.f, 16.f, 0.f)
 		.HAlign(HAlign_Right)
 		[
 			SNew(SButton)
@@ -52,8 +52,8 @@ void SCustomTextFilterDialog::Construct( const FArguments& InArgs )
 	{
 		// Button to create a new filter and enable it
 		ButtonBox->AddSlot()
-		.FillWidth(1.0)
-		.Padding(0, 0, 16, 0)
+		.FillWidth(1.0f)
+		.Padding(0.f, 0.f, 16.f, 0.f)
 		.HAlign(HAlign_Right)
 		[
 			SNew(SButton)
@@ -64,7 +64,7 @@ void SCustomTextFilterDialog::Construct( const FArguments& InArgs )
 		// Button to create a new filter
 		ButtonBox->AddSlot()
 		.AutoWidth()
-		.Padding(0, 0, 16, 0)
+		.Padding(0.f, 0.f, 16.f, 0.f)
 		.HAlign(HAlign_Right)
 		[
 			SNew(SPrimaryButton)
@@ -77,7 +77,7 @@ void SCustomTextFilterDialog::Construct( const FArguments& InArgs )
 	// Button to close the dialog box, common to both modes
 	ButtonBox->AddSlot()
 	.AutoWidth()
-	.Padding(0, 0, 16, 0)
+	.Padding(0.f, 0.f, 16.f, 0.f)
 	.HAlign(HAlign_Right)
 	[
 		SNew(SButton)
@@ -94,7 +94,7 @@ void SCustomTextFilterDialog::Construct( const FArguments& InArgs )
 
 			// Filter Name and Color
 			+ SVerticalBox::Slot()
-			.Padding(0, 38, 0, 0)
+			.Padding(0.f, 38.f, 0.f, 0.f)
 			.AutoHeight()
 			[
 				SNew(SHorizontalBox)
@@ -103,7 +103,7 @@ void SCustomTextFilterDialog::Construct( const FArguments& InArgs )
 				+ SHorizontalBox::Slot()
 				.AutoWidth()
 				.VAlign(VAlign_Center)
-				.Padding(34, 0, 0, 0)
+				.Padding(34.f, 0.f, 0.f, 0.f)
 				[
 					SNew(STextBlock)
 					.Text(LOCTEXT("FilterLabelText", "Filter Label"))
@@ -111,10 +111,10 @@ void SCustomTextFilterDialog::Construct( const FArguments& InArgs )
 				// Filter Name Input Box
 				+ SHorizontalBox::Slot()
 				.AutoWidth()
-				.Padding(21, 0, 0, 0)
+				.Padding(21.f, 0.f, 0.f, 0.f)
 				[
 					SNew(SBox)
-					.WidthOverride(254)
+					.WidthOverride(254.f)
 					[
 						SAssignNew(FilterLabelTextBox, SEditableTextBox)
 						.Text_Lambda([this] ()
@@ -133,7 +133,7 @@ void SCustomTextFilterDialog::Construct( const FArguments& InArgs )
 				+ SHorizontalBox::Slot()
 				.AutoWidth()
 				.VAlign(VAlign_Center)
-				.Padding(21, 0, 0, 0)
+				.Padding(21.f, 0.f, 0.f, 0.f)
 				[
 					SNew(STextBlock)
 					.Text(LOCTEXT("FilterColorText", "Color"))
@@ -141,7 +141,7 @@ void SCustomTextFilterDialog::Construct( const FArguments& InArgs )
 				// Filter Color
 				+ SHorizontalBox::Slot()
 				.AutoWidth()
-				.Padding(21, 0, 0, 0)
+				.Padding(21.f, 0.f, 0.f, 0.f)
 				[
 					SAssignNew(ColorBlock, SColorBlock )
 					.Color_Lambda([this] ()
@@ -158,14 +158,14 @@ void SCustomTextFilterDialog::Construct( const FArguments& InArgs )
 			// Filter String
 			+ SVerticalBox::Slot()
 			.AutoHeight()
-			.Padding(0, 21, 0, 0)
+			.Padding(0.f, 21.f, 0.f, 0.f)
 			[
 				SNew(SHorizontalBox)
 				// Filter String Text
 				+ SHorizontalBox::Slot()
 				.AutoWidth()
 				.VAlign(VAlign_Center)
-				.Padding(31, 0, 0, 0)
+				.Padding(31.f, 0.f, 0.f, 0.f)
 				[
 					SNew(STextBlock)
 					.Text(LOCTEXT("FilterStringText", "Filter String"))
@@ -173,10 +173,10 @@ void SCustomTextFilterDialog::Construct( const FArguments& InArgs )
 				// Filter Name Input Box
 				+ SHorizontalBox::Slot()
 				.AutoWidth()
-				.Padding(21, 0, 0, 0)
+				.Padding(21.f, 0.f, 0.f, 0.f)
 				[
 					SNew(SBox)
-					.WidthOverride(556)
+					.WidthOverride(556.f)
 					[
 						SNew(SEditableTextBox)
 						.Text_Lambda([this] ()
@@ -194,7 +194,7 @@ void SCustomTextFilterDialog::Construct( const FArguments& InArgs )
 			// Buttons
 			+ SVerticalBox::Slot()
 			.VAlign(VAlign_Bottom)
-			.Padding(0, 0, 0, 16)
+			.Padding(0.f, 0.f, 0.f, 16.f)
 			[
 				ButtonBox.ToSharedRef()
 			]
diff --git a/Engine/Source/Developer/ToolWidgets/Private/Filters/SFilterSearchBox.cpp b/Engine/Source/Developer/ToolWidgets/Private/Filters/SFilterSearchBox.cpp
--- a/Engine/Source/Developer/ToolWidgets/Private/Filters/SFilterSearchBox.cpp
+++ b/Engine/Source/Developer/ToolWidgets/Private/Filters/SFilterSearchBox.cpp
@@ -94,12 +94,12 @@ void SFilterSearchBoxImpl::Construct(const FArguments& InArgs)
 		.VAlign(VAlign_Center)
 		.HAlign(HAlign_Right)
 		.AutoWidth()
-		.Padding(-4, 0, 0, 0) // Negative padding intentional to have the same padding between all buttons
+		.Padding(-4.f, 0.f, 0.f, 0.f) // Negative padding intentional to have the same padding between all buttons
 		[
 			// Button to save the currently occuring search 
 			SNew(SButton)
 			.AddMetaData(FDriverMetaData::Id(TEXT("CreateFilterFromCurrentSearchButton")))
-			.ContentPadding(0)
+			.ContentPadding(0.f)
 			.Visibility_Lambda([this]()
 			{
 				// Only visible if there is a search active currently and the OnSaveSearchClicked delegate is bound
@@ -123,11 +123,11 @@ void SFilterSearchBoxImpl::Construct(const FArguments& InArgs)
 		.HAlign(HAlign_Right)
 		.VAlign(VAlign_Center)
 		.AutoWidth()
-		.Padding(0, 0, 4, 0)
+		.Padding(0.f, 0.f, 4.f, 0.f)
 		[
 			// Chevron to open the search history dropdown
 			SNew(SButton)
-			.ContentPadding(0)
+			.ContentPadding(0.f)
 			.ButtonStyle(FAppStyle::Get(), "HoverOnlyButton")
 			.ClickMethod(EButtonClickMethod::MouseDown)
 			.OnClicked(InArgs._OnSearchHistoryClicked)
@@ -193,7 +193,7 @@ void SFilterSearchBox::Construct( const FArguments& InArgs )
 		(
 			SNew(SBorder)
 			.BorderImage(FAppStyle::GetBrush("Menu.Background"))
-			.Padding( FMargin(2) )
+			.Padding( FMargin(2.f) )
 			[
 				SAssignNew(SearchHistoryListView, SListView< TSharedPtr<FText> >)
 				.ListItemsSource(&SearchHistory)
@@ -308,7 +308,7 @@ TSharedRef<ITableRow> SFilterSearchBox::MakeSearchHistoryRowWidget(TSharedPtr<FT
 	.HAlign(TextAlignment)
 	.VAlign(VAlign_Center)
 	.MaxWidth(500.0f)
-	.FillWidth(1.0)
+	.FillWidth(1.0f)
 	[
 		SNew(STextBlock)
 		.Text(*SearchText.Get())
diff --git a/Engine/Source/Developer/ToolWidgets/Private/SResizeBox.cpp b/Engine/Source/Developer/ToolWidgets/Private/SResizeBox.cpp
--- a/Engine/Source/Developer/ToolWidgets/Private/SResizeBox.cpp
+++ b/Engine/Source/Developer/ToolWidgets/Private/SResizeBox.cpp
@@ -28,7 +28,7 @@ FReply SVerticalResizeBox::OnMouseButtonDown(const FGeometry& MyGeometry, const
 {
 	if (MouseEvent.GetEffectingButton() == EKeys::LeftMouseButton)
 	{
-		FVector2D MouseLocation = MyGeometry.AbsoluteToLocal(MouseEvent.GetScreenSpacePosition());
+		FVector2f MouseLocation = MyGeometry.AbsoluteToLocal(MouseEvent.GetScreenSpacePosition());
 		if (MyGeometry.GetLocalSize().Y - MouseLocation.Y < HandleHeight)
 		{
 			DragStartLocation = MouseLocation.Y;
@@ -50,7 +50,7 @@ FReply SVerticalResizeBox::OnMouseButtonUp(const FGeometry& MyGeometry, const FP
 
 FReply SVerticalResizeBox::OnMouseMove(const FGeometry& MyGeometry, const FPointerEvent& MouseEvent)
 {
-	FVector2D MouseLocation = MyGeometry.AbsoluteToLocal(MouseEvent.GetScreenSpacePosition());
+	FVector2f MouseLocation = MyGeometry.AbsoluteToLocal(MouseEvent.GetScreenSpacePosition());
 	LastMouseLocation = MouseLocation.Y;
 	if (this->HasMouseCapture())
 	{
@@ -70,7 +70,7 @@ FReply SVerticalResizeBox::OnMouseMove(const FGeometry& MyGeometry, const FPoint
 
 FCursorReply SVerticalResizeBox::OnCursorQuery(const FGeometry& MyGeometry, const FPointerEvent& CursorEvent) const
 {
-	FVector2D CursorLocation = MyGeometry.AbsoluteToLocal(CursorEvent.GetScreenSpacePosition());
+	FVector2f CursorLocation = MyGeometry.AbsoluteToLocal(CursorEvent.GetScreenSpacePosition());
 	if (MyGeometry.GetLocalSize().Y - CursorLocation.Y < HandleHeight)
 	{
 		return FCursorReply::Cursor(EMouseCursor::ResizeUpDown);
@@ -82,7 +82,7 @@ int32 SVerticalResizeBox::OnPaint(const FPaintArgs& Args, const FGeometry& Allot
 {
 	FLinearColor HandleBoxColor;
 	int32 HandleLayerId = LayerId + 1;
-	FVector2D LocalSize = AllottedGeometry.GetLocalSize();
+	FVector2f LocalSize = AllottedGeometry.GetLocalSize();
 	if (IsHovered() && LastMouseLocation.IsSet() && LastMouseLocation.GetValue() >= LocalSize.Y - HandleHeight && LastMouseLocation.GetValue() <= LocalSize.Y)
 	{
 		HandleBoxColor = HandleHighlightColor.Get();
@@ -92,8 +92,8 @@ int32 SVerticalResizeBox::OnPaint(const FPaintArgs& Args, const FGeometry& Allot
 		HandleBoxColor = HandleColor.Get();
 	}
 
-	FVector2D HandleLocation(0, AllottedGeometry.GetLocalSize().Y - HandleHeight);
-	FVector2D HandleSize(AllottedGeometry.GetLocalSize().X, HandleHeight);
+	FVector2f HandleLocation(0, AllottedGeometry.GetLocalSize().Y - HandleHeight);
+	FVector2f HandleSize(AllottedGeometry.GetLocalSize().X, HandleHeight);
 	FSlateDrawElement::MakeBox
 	(
 		OutDrawElements,
diff --git a/Engine/Source/Developer/ToolWidgets/Private/SSegmentedTimelineView.cpp b/Engine/Source/Developer/ToolWidgets/Private/SSegmentedTimelineView.cpp
--- a/Engine/Source/Developer/ToolWidgets/Private/SSegmentedTimelineView.cpp
+++ b/Engine/Source/Developer/ToolWidgets/Private/SSegmentedTimelineView.cpp
@@ -24,7 +24,7 @@ int32 SSegmentedTimelineView::OnPaint(const FPaintArgs& Args, const FGeometry& A
 
 FVector2D SSegmentedTimelineView::ComputeDesiredSize(float) const
 {
-	return DesiredSize.Get();
+	return UE::Slate::FDeprecateVector2DResult(DesiredSize.Get());
 }
 
 int32 SSegmentedTimelineView::PaintBlock(const FGeometry& AllottedGeometry, const FSlateRect& MyCullingRect, FSlateWindowElementList& OutDrawElements, int32 LayerId, const FWidgetStyle& InWidgetStyle, bool bParentEnabled) const
@@ -34,7 +34,7 @@ int32 SSegmentedTimelineView::PaintBlock(const FGeometry& AllottedGeometry, cons
 	FLinearColor Color = FillColor.Get();
 
 	SSimpleTimeSlider::FScrubRangeToScreen RangeToScreen(DebugTimeRange, AllottedGeometry.GetLocalSize());
-	FVector2D Size = AllottedGeometry.GetLocalSize();
+	FVector2f Size = AllottedGeometry.GetLocalSize();
 
 	if (TSharedPtr<FSegmentData> SegmentDataPtr = SegmentData.Get())
 	{
@@ -45,7 +45,7 @@ int32 SSegmentedTimelineView::PaintBlock(const FGeometry& AllottedGeometry, cons
 		{
 			double LowerBound = SegmentRange.HasLowerBound() ? SegmentRange.GetLowerBoundValue() : 0;
 			double UpperBound = SegmentRange.HasUpperBound() ? SegmentRange.GetUpperBoundValue() : Size.X;
-			float BoxMin = FMath::Max(0, RangeToScreen.InputToLocalX(LowerBound));
+			float BoxMin = FMath::Max(0.f, RangeToScreen.InputToLocalX(LowerBound));
 			float BoxMax = FMath::Min(Size.X, RangeToScreen.InputToLocalX(UpperBound));
 
 			FPaintGeometry BoxGeometry = AllottedGeometry.ToPaintGeometry(FVector2f(BoxMax - BoxMin, Size.Y - 2), FSlateLayoutTransform(FVector2f(BoxMin, 1.f)));
diff --git a/Engine/Source/Developer/ToolWidgets/Private/SSimpleTimeSlider.cpp b/Engine/Source/Developer/ToolWidgets/Private/SSimpleTimeSlider.cpp
--- a/Engine/Source/Developer/ToolWidgets/Private/SSimpleTimeSlider.cpp
+++ b/Engine/Source/Developer/ToolWidgets/Private/SSimpleTimeSlider.cpp
@@ -45,7 +45,7 @@ int32 SSimpleTimeSlider::OnPaint( const FPaintArgs& Args, const FGeometry& Allot
 
 FVector2D SSimpleTimeSlider::ComputeDesiredSize( float ) const
 {
-	return DesiredSize;
+	return UE::Slate::FDeprecateVector2DResult(DesiredSize);
 }
 
 namespace ScrubConstants
@@ -85,7 +85,7 @@ static float GetNextSpacing( uint32 CurrentStep )
 * @param MinTickSpacing	The minimum tick spacing in time units allowed
 * @return the optimal spacing in time units
 */
-static double DetermineOptimalSpacing(float InPixelsPerInput, uint32 InMinTick, double InMinTickSpacing)
+static float DetermineOptimalSpacing(float InPixelsPerInput, uint32 InMinTick, float InMinTickSpacing)
 {
 	if (InPixelsPerInput <= 0.0f)
 		return InMinTickSpacing;
@@ -93,8 +93,8 @@ static double DetermineOptimalSpacing(float InPixelsPerInput, uint32 InMinTick,
 	uint32 CurStep = 0;
 
 	// Start with the smallest spacing
-	double Spacing = InMinTickSpacing;
-	double MinTick = InMinTick;
+	float Spacing = InMinTickSpacing;
+	float MinTick = static_cast<float>(InMinTick);
 
 	while( Spacing * InPixelsPerInput < MinTick )
 	{
@@ -116,15 +116,15 @@ void SSimpleTimeSlider::DrawTicks( FSlateWindowElementList& OutDrawElements, con
 	// For slightly larger halfway tick mark
 	const int32 HalfDivider = Divider / 2;
 	// Find out where to start from
-	int32 OffsetNum = FMath::FloorToInt(RangeToScreen.ViewInput.GetLowerBoundValue() / Spacing);
+	int32 OffsetNum = FMath::FloorToInt32(RangeToScreen.ViewInput.GetLowerBoundValue() / Spacing);
 	
 	FSlateFontInfo SmallLayoutFont = FCoreStyle::GetDefaultFontStyle("Regular", 8);
 
-	TArray<FVector2D> LinePoints;
+	TArray<FVector2f> LinePoints;
 	LinePoints.AddUninitialized(2);
 
 	double Seconds = 0;
-	while( (Seconds = OffsetNum*Spacing) < RangeToScreen.ViewInput.GetUpperBoundValue() )
+	while( (Seconds = static_cast<float>(OffsetNum)*Spacing) < RangeToScreen.ViewInput.GetUpperBoundValue() )
 	{
 		// X position local to start of the widget area
 		float XPos = RangeToScreen.InputToLocalX( Seconds );
@@ -135,8 +135,8 @@ void SSimpleTimeSlider::DrawTicks( FSlateWindowElementList& OutDrawElements, con
 			FVector2f Offset( XPos, InArgs.TickOffset );
 			FVector2f TickSize( 1.0f, InArgs.MajorTickHeight );
 
-			LinePoints[0] = FVector2D(1.0f,1.0f);
-			LinePoints[1] = FVector2D(TickSize);
+			LinePoints[0] = FVector2f(1.0f,1.0f);
+			LinePoints[1] = FVector2f(TickSize);
 
 			// lines should not need anti-aliasing
 			const bool bAntiAliasLines = false;
@@ -180,8 +180,8 @@ void SSimpleTimeSlider::DrawTicks( FSlateWindowElementList& OutDrawElements, con
 			FVector2f Offset(XPos, InArgs.bMirrorLabels ? 0.0f : FMath::Abs( InArgs.AllottedGeometry.GetLocalSize().Y - MinorTickHeight ) );
 			FVector2f TickSize(1.f, MinorTickHeight);
 
-			LinePoints[0] = FVector2D(1.0f,1.0f);
-			LinePoints[1] = FVector2D(TickSize);
+			LinePoints[0] = FVector2f(1.0f,1.0f);
+			LinePoints[1] = FVector2f(TickSize);
 
 			const bool bAntiAlias = false;
 			// Draw each sub mark
@@ -206,12 +206,12 @@ int32 SSimpleTimeSlider::OnPaintTimeSlider( bool bMirrorLabels, const FGeometry&
 	const ESlateDrawEffect DrawEffects = bEnabled ? ESlateDrawEffect::None : ESlateDrawEffect::DisabledEffect;
 
 	TRange<double> LocalViewRange = ViewRange.Get();
-	const float LocalViewRangeMin = LocalViewRange.GetLowerBoundValue();
-	const float LocalViewRangeMax = LocalViewRange.GetUpperBoundValue();
-	const float LocalSequenceLength = LocalViewRangeMax-LocalViewRangeMin;
+	const double LocalViewRangeMin = LocalViewRange.GetLowerBoundValue();
+	const double LocalViewRangeMax = LocalViewRange.GetUpperBoundValue();
+	const double LocalSequenceLength = LocalViewRangeMax-LocalViewRangeMin;
 	
-	FVector2D Scale = FVector2D(1.0f,1.0f);
-	if ( LocalSequenceLength > 0)
+	FVector2f Scale = FVector2f(1.0f,1.0f);
+	if ( LocalSequenceLength > 0.0)
 	{
 		FScrubRangeToScreen RangeToScreen( LocalViewRange, AllottedGeometry.GetLocalSize() );
 		
@@ -317,7 +317,7 @@ FReply SSimpleTimeSlider::OnPreviewMouseButtonDown(const FGeometry& MyGeometry,
 	if ( bHandleLeftMouseButton )
 	{
 		// Always capture mouse if we left or right click on the widget
-		const float NewValue = GetTimeAtCursorPosition(MyGeometry, MouseEvent);
+		const double NewValue = GetTimeAtCursorPosition(MyGeometry, MouseEvent);
 		CommitScrubPosition(NewValue, /*bIsScrubbing=*/false);
 		return FReply::Handled().CaptureMouse(AsShared()).PreventThrottling();
 	}
@@ -349,7 +349,7 @@ FReply SSimpleTimeSlider::OnMouseButtonUp(const FGeometry& MyGeometry, const FPo
 	{
 		if( !bDraggingScrubber )
 		{
-			const float NewValue = GetTimeAtCursorPosition(MyGeometry, MouseEvent);
+			const double NewValue = GetTimeAtCursorPosition(MyGeometry, MouseEvent);
 			CommitScrubPosition( NewValue, /*bIsScrubbing=*/false );
 		}
 
@@ -361,14 +361,14 @@ FReply SSimpleTimeSlider::OnMouseButtonUp(const FGeometry& MyGeometry, const FPo
 	return FReply::Unhandled();
 }
 
-float SSimpleTimeSlider::GetTimeAtCursorPosition(const FGeometry& MyGeometry, const FPointerEvent& MouseEvent) const
+double SSimpleTimeSlider::GetTimeAtCursorPosition(const FGeometry& MyGeometry, const FPointerEvent& MouseEvent) const
 {
 	FScrubRangeToScreen RangeToScreen(ViewRange.Get(), MyGeometry.GetLocalSize());
-	FVector2D CursorPos = MyGeometry.AbsoluteToLocal(MouseEvent.GetLastScreenSpacePosition());
-	float NewValue = RangeToScreen.LocalXToInput(CursorPos.X);
+	FVector2f CursorPos = MyGeometry.AbsoluteToLocal(MouseEvent.GetLastScreenSpacePosition());
+	double NewValue = RangeToScreen.LocalXToInput(CursorPos.X);
 
-	float LocalClampMin = ClampRange.Get().GetLowerBoundValue();
-	float LocalClampMax = ClampRange.Get().GetUpperBoundValue();
+	double LocalClampMin = ClampRange.Get().GetLowerBoundValue();
+	double LocalClampMax = ClampRange.Get().GetUpperBoundValue();
 
 	if (NewValue < LocalClampMin)
 	{
@@ -405,14 +405,14 @@ FReply SSimpleTimeSlider::OnMouseMove( const FGeometry& MyGeometry, const FPoint
 
 				TRange<double> LocalViewRange = ViewRange.Get();
 				double LocalViewRangeMin = LocalViewRange.GetLowerBoundValue();
-				float LocalViewRangeMax = LocalViewRange.GetUpperBoundValue();
+				double LocalViewRangeMax = LocalViewRange.GetUpperBoundValue();
 
 				FScrubRangeToScreen ScaleInfo( LocalViewRange, MyGeometry.GetLocalSize() );
-				FVector2D ScreenDelta = MouseEvent.GetCursorDelta();
-				FVector2D InputDelta;
+				FVector2f ScreenDelta = MouseEvent.GetCursorDelta();
+				FVector2f InputDelta;
 				InputDelta.X = ScreenDelta.X/ScaleInfo.PixelsPerInput;
 
-				const TRange<double> NewViewRange = TRange<double>(LocalViewRangeMin - InputDelta.X, LocalViewRangeMax - InputDelta.X);
+				const TRange<double> NewViewRange = TRange<double>(LocalViewRangeMin - static_cast<double>(InputDelta.X), LocalViewRangeMax - static_cast<double>(InputDelta.X));
 				TRange<double> LocalClampRange = ClampRange.Get();
 
 				OnViewRangeChanged.ExecuteIfBound(NewViewRange);
@@ -436,7 +436,7 @@ FReply SSimpleTimeSlider::OnMouseMove( const FGeometry& MyGeometry, const FPoint
 			}
 			else
 			{
-				const float NewValue = GetTimeAtCursorPosition(MyGeometry, MouseEvent);
+				const double NewValue = GetTimeAtCursorPosition(MyGeometry, MouseEvent);
 				CommitScrubPosition(NewValue, /*bIsScrubbing=*/true);
 			}
 		}
@@ -446,7 +446,7 @@ FReply SSimpleTimeSlider::OnMouseMove( const FGeometry& MyGeometry, const FPoint
 	return FReply::Unhandled();
 }
 
-void SSimpleTimeSlider::CommitScrubPosition( float NewValue, bool bIsScrubbing )
+void SSimpleTimeSlider::CommitScrubPosition( double NewValue, bool bIsScrubbing )
 {
 	// Manage the scrub position ourselves if its not bound to a delegate
 	if ( !ScrubPosition.IsBound() )
diff --git a/Engine/Source/Developer/ToolWidgets/Private/SWidgetDrawer.cpp b/Engine/Source/Developer/ToolWidgets/Private/SWidgetDrawer.cpp
--- a/Engine/Source/Developer/ToolWidgets/Private/SWidgetDrawer.cpp
+++ b/Engine/Source/Developer/ToolWidgets/Private/SWidgetDrawer.cpp
@@ -33,7 +33,7 @@ class SDrawerOverlay : public SCompoundWidget
 	SLATE_BEGIN_ARGS(SDrawerOverlay)
 	{
 		_Clipping = EWidgetClipping::ClipToBounds;
-		_ShadowOffset = FVector2D(10.0f, 20.0f);
+		_ShadowOffset = FVector2f(10.0f, 20.0f);
 	}
 		SLATE_DEFAULT_SLOT(FArguments, Content)
 		SLATE_ARGUMENT(float, MinDrawerHeight)
@@ -41,7 +41,7 @@ class SDrawerOverlay : public SCompoundWidget
 		SLATE_ARGUMENT(float, TargetDrawerHeight)
 		SLATE_EVENT(FOnStatusBarDrawerTargetHeightChanged, OnTargetHeightChanged)
 		SLATE_EVENT(FSimpleDelegate, OnDismissComplete)
-		SLATE_ARGUMENT(FVector2D, ShadowOffset)
+		SLATE_ARGUMENT(FDeprecateSlateVector2D, ShadowOffset)
 	SLATE_END_ARGS()
 
 	~SDrawerOverlay()
@@ -97,7 +97,7 @@ class SDrawerOverlay : public SCompoundWidget
 
 	virtual FVector2D ComputeDesiredSize(float) const
 	{
-		return FVector2D(1.0f, TargetHeight + ShadowOffset.Y);
+		return UE::Slate::FDeprecateVector2DResult(1.0f, TargetHeight + ShadowOffset.Y);
 	}
 
 	virtual void OnArrangeChildren(const FGeometry& AllottedGeometry, FArrangedChildren& ArrangedChildren) const override
@@ -109,7 +109,7 @@ class SDrawerOverlay : public SCompoundWidget
 				AllottedGeometry.MakeChild(
 					ChildSlot.GetWidget(),
 					ShadowOffset,
-					FVector2D(AllottedGeometry.GetLocalSize().X - (ShadowOffset.X * 2), TargetHeight)
+					FVector2f(AllottedGeometry.GetLocalSize().X - (ShadowOffset.X * 2), TargetHeight)
 				)
 			);
 		}
@@ -159,8 +159,8 @@ class SDrawerOverlay : public SCompoundWidget
 
 		if (bIsResizing && this->HasMouseCapture() && !MouseEvent.GetCursorDelta().IsZero())
 		{
-			const FVector2D LocalMousePos = InitialResizeGeometry.AbsoluteToLocal(MouseEvent.GetScreenSpacePosition());
-			const float DeltaHeight = (InitialResizeGeometry.GetLocalPositionAtCoordinates(FVector2D::ZeroVector) - LocalMousePos).Y;
+			const FVector2f LocalMousePos = InitialResizeGeometry.AbsoluteToLocal(MouseEvent.GetScreenSpacePosition());
+			const float DeltaHeight = (InitialResizeGeometry.GetLocalPositionAtCoordinates(FVector2f::ZeroVector) - LocalMousePos).Y;
 
 			TargetHeight = FMath::Clamp(InitialHeightAtResize + DeltaHeight, MinHeight, MaxHeight);
 			SetHeight(InitialHeightAtResize + DeltaHeight);
@@ -275,14 +275,14 @@ class SDrawerOverlay : public SCompoundWidget
 private:
 	FGeometry GetRenderTransformedGeometry(const FGeometry& AllottedGeometry) const
 	{
-		return AllottedGeometry.MakeChild(FSlateRenderTransform(FVector2D(0.0f, TargetHeight - CurrentHeight)));
+		return AllottedGeometry.MakeChild(FSlateRenderTransform(FVector2f(0.0f, TargetHeight - CurrentHeight)));
 	}
 
 	FGeometry GetResizeHandleGeometry(const FGeometry& AllottedGeometry) const
 	{
 		return GetRenderTransformedGeometry(AllottedGeometry).MakeChild(
-			FVector2D(AllottedGeometry.GetLocalSize().X-ShadowOffset.X*2, ExpanderSize),
-			FSlateLayoutTransform(ShadowOffset - FVector2D(0.0f, ExpanderSize))
+			FVector2f(AllottedGeometry.GetLocalSize().X-ShadowOffset.X*2.f, ExpanderSize),
+			FSlateLayoutTransform(ShadowOffset - FVector2f(0.0f, ExpanderSize))
 		);
 	}
 
@@ -321,7 +321,7 @@ private:
 	const FSlateBrush* ShadowBrush;
 	const FSlateBrush* BorderBrush;
 	const FSplitterStyle* SplitterStyle;
-	FVector2D ShadowOffset;
+	FVector2f ShadowOffset;
 	FThrottleRequest AnimationThrottle;
 	FThrottleRequest ResizeThrottleHandle;
 	float ExpanderSize;
diff --git a/Engine/Source/Developer/ToolWidgets/Private/Sidebar/SSidebarContainer.cpp b/Engine/Source/Developer/ToolWidgets/Private/Sidebar/SSidebarContainer.cpp
--- a/Engine/Source/Developer/ToolWidgets/Private/Sidebar/SSidebarContainer.cpp
+++ b/Engine/Source/Developer/ToolWidgets/Private/Sidebar/SSidebarContainer.cpp
@@ -1,4 +1,4 @@
-﻿// Copyright Epic Games, Inc. All Rights Reserved.
+// Copyright Epic Games, Inc. All Rights Reserved.
 
 #include "Sidebar/SSidebarContainer.h"
 #include "Framework/Application/SlateApplication.h"
@@ -308,7 +308,7 @@ float SSidebarContainer::GetCurrentDrawerSize() const
 	return SidebarSizePercent;
 }
 
-FVector2D SSidebarContainer::GetOverlaySize() const
+UE::Slate::FDeprecateVector2DResult SSidebarContainer::GetOverlaySize() const
 {
 	return DrawersOverlay->GetTickSpaceGeometry().GetLocalSize();
 }
diff --git a/Engine/Source/Developer/ToolWidgets/Private/Sidebar/SSidebarDrawer.cpp b/Engine/Source/Developer/ToolWidgets/Private/Sidebar/SSidebarDrawer.cpp
--- a/Engine/Source/Developer/ToolWidgets/Private/Sidebar/SSidebarDrawer.cpp
+++ b/Engine/Source/Developer/ToolWidgets/Private/Sidebar/SSidebarDrawer.cpp
@@ -139,26 +139,26 @@ void SSidebarDrawer::OnArrangeChildren(const FGeometry& InAllottedGeometry, FArr
 		return;
 	}
 
-	FVector2D ChildOffset;
-	FVector2D LocalSize;
+	FVector2f ChildOffset;
+	FVector2f LocalSize;
 
 	switch (TabLocation)
 	{
 	case ESidebarTabLocation::Left:
-		ChildOffset = FVector2D(0.f, ShadowOffset.Y);
-		LocalSize = FVector2D(TargetDrawerSize, InAllottedGeometry.GetLocalSize().Y - (ShadowOffset.Y * 2));
+		ChildOffset = FVector2f(0.f, ShadowOffset.Y);
+		LocalSize = FVector2f(TargetDrawerSize, InAllottedGeometry.GetLocalSize().Y - (ShadowOffset.Y * 2));
 		break;
 	case ESidebarTabLocation::Right:
 		ChildOffset = ShadowOffset;
-		LocalSize = FVector2D(TargetDrawerSize, InAllottedGeometry.GetLocalSize().Y - (ShadowOffset.Y * 2));
+		LocalSize = FVector2f(TargetDrawerSize, InAllottedGeometry.GetLocalSize().Y - (ShadowOffset.Y * 2));
 		break;
 	case ESidebarTabLocation::Top:
-		ChildOffset = FVector2D(ShadowOffset.X, 0.f);
-		LocalSize = FVector2D(InAllottedGeometry.GetLocalSize().X - (ShadowOffset.X * 2), TargetDrawerSize);
+		ChildOffset = FVector2f(ShadowOffset.X, 0.f);
+		LocalSize = FVector2f(InAllottedGeometry.GetLocalSize().X - (ShadowOffset.X * 2), TargetDrawerSize);
 		break;
 	case ESidebarTabLocation::Bottom:
 		ChildOffset = ShadowOffset;
-		LocalSize = FVector2D(InAllottedGeometry.GetLocalSize().X - (ShadowOffset.X * 2), TargetDrawerSize);
+		LocalSize = FVector2f(InAllottedGeometry.GetLocalSize().X - (ShadowOffset.X * 2), TargetDrawerSize);
 		break;
 	}
 
@@ -210,8 +210,8 @@ FReply SSidebarDrawer::OnMouseMove(const FGeometry& InAllottedGeometry, const FP
 
 	if (bIsResizing && this->HasMouseCapture() && !InMouseEvent.GetCursorDelta().IsZero())
 	{
-		const FVector2D MousePosition = InMouseEvent.GetScreenSpacePosition();
-		const FVector2D LocalMousePosition = InitialResizeGeometry.AbsoluteToLocal(MousePosition);
+		const FVector2f MousePosition = InMouseEvent.GetScreenSpacePosition();
+		const FVector2f LocalMousePosition = InitialResizeGeometry.AbsoluteToLocal(MousePosition);
 
 		float DeltaSize = 0.f;
 
@@ -291,25 +291,25 @@ int32 SSidebarDrawer::OnPaint(const FPaintArgs& InArgs, const FGeometry& InAllot
 
 	const FGeometry RenderTransformedChildGeometry = GetRenderTransformedGeometry(InAllottedGeometry);
 	const FGeometry ResizeHandleGeometry = GetResizeHandleGeometry(InAllottedGeometry);
-	const FVector2D LocalSize = InAllottedGeometry.GetLocalSize();
+	const FVector2f LocalSize = InAllottedGeometry.GetLocalSize();
 
-	FVector2D ContentsLocalOrigin;
-	FVector2D ContentsLocalSize;
+	FVector2f ContentsLocalOrigin;
+	FVector2f ContentsLocalSize;
 
 	switch (TabLocation)
 	{
 	case ESidebarTabLocation::Left:
-		ContentsLocalOrigin = FVector2D(0.f, ShadowOffset.Y);
-		ContentsLocalSize = FVector2D(TargetDrawerSize, LocalSize.Y - (ShadowOffset.Y * 2));
+		ContentsLocalOrigin = FVector2f(0.f, ShadowOffset.Y);
+		ContentsLocalSize = FVector2f(TargetDrawerSize, LocalSize.Y - (ShadowOffset.Y * 2));
 		break;
 	case ESidebarTabLocation::Right:
 		ContentsLocalOrigin = ShadowOffset;
-		ContentsLocalSize = FVector2D(TargetDrawerSize, LocalSize.Y - (ShadowOffset.Y * 2));
+		ContentsLocalSize = FVector2f(TargetDrawerSize, LocalSize.Y - (ShadowOffset.Y * 2));
 		break;
 	case ESidebarTabLocation::Top:
 	case ESidebarTabLocation::Bottom:
 		ContentsLocalOrigin = ShadowOffset;
-		ContentsLocalSize = FVector2D(LocalSize.X - (ShadowOffset.X * 2), TargetDrawerSize);
+		ContentsLocalSize = FVector2f(LocalSize.X - (ShadowOffset.X * 2), TargetDrawerSize);
 		break;
 	}
 
@@ -391,12 +391,12 @@ int32 SSidebarDrawer::OnPaint(const FPaintArgs& InArgs, const FGeometry& InAllot
 
 			// Compute the top/bottom of the tab in our local space.
 			const float BorderWidth = BorderBrush->OutlineSettings.Width;
-			const float TabTopY = RenderTransformedChildGeometry.AbsoluteToLocal(TabButtonGeometry.GetAbsolutePositionAtCoordinates(FVector2D::ZeroVector)).Y + 0.5f * BorderWidth;
-			const float TabBottomY = RenderTransformedChildGeometry.AbsoluteToLocal(TabButtonGeometry.GetAbsolutePositionAtCoordinates(FVector2D::UnitVector)).Y - 0.5f * BorderWidth;
+			const float TabTopY = RenderTransformedChildGeometry.AbsoluteToLocal(TabButtonGeometry.GetAbsolutePositionAtCoordinates(FVector2f::ZeroVector)).Y + 0.5f * BorderWidth;
+			const float TabBottomY = RenderTransformedChildGeometry.AbsoluteToLocal(TabButtonGeometry.GetAbsolutePositionAtCoordinates(FVector2f::UnitVector)).Y - 0.5f * BorderWidth;
 
 			// Create the geometry for the notched portion, where one edge extends past the clipping rect.
-			const FVector2D NotchOffsetSize(TabButtonGeometry.GetLocalSize().X, 0.f);
-			const FVector2D NotchOffsetTranslate = TabLocation == ESidebarTabLocation::Left ? -NotchOffsetSize : FVector2D::ZeroVector;
+			const FVector2f NotchOffsetSize(TabButtonGeometry.GetLocalSize().X, 0.f);
+			const FVector2f NotchOffsetTranslate = TabLocation == ESidebarTabLocation::Left ? -NotchOffsetSize : FVector2f::ZeroVector;
 			const FPaintGeometry NotchOffsetPaintGeom = RenderTransformedChildGeometry.ToPaintGeometry(ContentsLocalSize + NotchOffsetSize
 				, FSlateLayoutTransform(ContentsLocalOrigin + NotchOffsetTranslate));
 
@@ -456,13 +456,13 @@ FGeometry SSidebarDrawer::GetRenderTransformedGeometry(const FGeometry& InAllott
 	switch (TabLocation)
 	{
 	case ESidebarTabLocation::Left:
-		return InAllottedGeometry.MakeChild(FSlateRenderTransform(FVector2D(CurrentSize - TargetDrawerSize, 0.f)));
+		return InAllottedGeometry.MakeChild(FSlateRenderTransform(FVector2f(CurrentSize - TargetDrawerSize, 0.f)));
 	case ESidebarTabLocation::Right:
-		return InAllottedGeometry.MakeChild(FSlateRenderTransform(FVector2D(TargetDrawerSize - CurrentSize, 0.f)));
+		return InAllottedGeometry.MakeChild(FSlateRenderTransform(FVector2f(TargetDrawerSize - CurrentSize, 0.f)));
 	case ESidebarTabLocation::Top:
-		return InAllottedGeometry.MakeChild(FSlateRenderTransform(FVector2D(0.f, CurrentSize - TargetDrawerSize)));
+		return InAllottedGeometry.MakeChild(FSlateRenderTransform(FVector2f(0.f, CurrentSize - TargetDrawerSize)));
 	case ESidebarTabLocation::Bottom:
-		return InAllottedGeometry.MakeChild(FSlateRenderTransform(FVector2D(0.f, TargetDrawerSize - CurrentSize)));
+		return InAllottedGeometry.MakeChild(FSlateRenderTransform(FVector2f(0.f, TargetDrawerSize - CurrentSize)));
 	}
 	return FGeometry();
 }
@@ -471,26 +471,26 @@ FGeometry SSidebarDrawer::GetResizeHandleGeometry(const FGeometry& InAllottedGeo
 {
 	const FGeometry RenderTransformedGeometry = GetRenderTransformedGeometry(InAllottedGeometry);
 
-	FVector2D LocalSize;
-	FVector2D Translation;
+	FVector2f LocalSize;
+	FVector2f Translation;
 
 	switch (TabLocation)
 	{
 	case ESidebarTabLocation::Left:
-		LocalSize = FVector2D(ExpanderHandleSize, InAllottedGeometry.GetLocalSize().Y - (ShadowOffset.Y * 2));
-		Translation = FVector2D(RenderTransformedGeometry.GetLocalSize().X - ShadowOffset.X, ShadowOffset.Y);
+		LocalSize = FVector2f(ExpanderHandleSize, InAllottedGeometry.GetLocalSize().Y - (ShadowOffset.Y * 2));
+		Translation = FVector2f(RenderTransformedGeometry.GetLocalSize().X - ShadowOffset.X, ShadowOffset.Y);
 		break;
 	case ESidebarTabLocation::Right:
-		LocalSize = FVector2D(ExpanderHandleSize, InAllottedGeometry.GetLocalSize().Y - (ShadowOffset.Y * 2));
-		Translation = ShadowOffset - FVector2D(ExpanderHandleSize, 0.f);
+		LocalSize = FVector2f(ExpanderHandleSize, InAllottedGeometry.GetLocalSize().Y - (ShadowOffset.Y * 2));
+		Translation = ShadowOffset - FVector2f(ExpanderHandleSize, 0.f);
 		break;
 	case ESidebarTabLocation::Top:
-		LocalSize = FVector2D(InAllottedGeometry.GetLocalSize().X - (ShadowOffset.X * 2), ExpanderHandleSize);
-		Translation = FVector2D(ShadowOffset.X, RenderTransformedGeometry.GetLocalSize().Y - ShadowOffset.Y);
+		LocalSize = FVector2f(InAllottedGeometry.GetLocalSize().X - (ShadowOffset.X * 2), ExpanderHandleSize);
+		Translation = FVector2f(ShadowOffset.X, RenderTransformedGeometry.GetLocalSize().Y - ShadowOffset.Y);
 		break;
 	case ESidebarTabLocation::Bottom:
-		LocalSize = FVector2D(InAllottedGeometry.GetLocalSize().X - (ShadowOffset.X * 2), ExpanderHandleSize);
-		Translation = ShadowOffset - FVector2D(0.f, ExpanderHandleSize);
+		LocalSize = FVector2f(InAllottedGeometry.GetLocalSize().X - (ShadowOffset.X * 2), ExpanderHandleSize);
+		Translation = ShadowOffset - FVector2f(0.f, ExpanderHandleSize);
 		break;
 	}
 
diff --git a/Engine/Source/Developer/ToolWidgets/Private/Sidebar/SSidebarDrawer.h b/Engine/Source/Developer/ToolWidgets/Private/Sidebar/SSidebarDrawer.h
--- a/Engine/Source/Developer/ToolWidgets/Private/Sidebar/SSidebarDrawer.h
+++ b/Engine/Source/Developer/ToolWidgets/Private/Sidebar/SSidebarDrawer.h
@@ -24,7 +24,7 @@ class SSidebarDrawer : public SCompoundWidget
 {
 public:
 	SLATE_BEGIN_ARGS(SSidebarDrawer)
-		: _ShadowOffset(FVector2D(8.f, 8.f))
+		: _ShadowOffset(FVector2f(8.f, 8.f))
 		, _ExpanderHandleSize(5.f)
 	{
 		_Clipping = EWidgetClipping::ClipToBounds;
@@ -36,7 +36,7 @@ public:
 		/** The size that the drawer should be, clamped to the above min/max values. This unit is in window space */
 		SLATE_ARGUMENT(float, TargetDrawerSize)
 		/** The side of the drop shadow surrounding the drawer */
-		SLATE_ARGUMENT(FVector2D, ShadowOffset)
+		SLATE_ARGUMENT(FDeprecateSlateVector2D, ShadowOffset)
 		/** The size of the handle used to resize the drawer */
 		SLATE_ARGUMENT(float, ExpanderHandleSize)
 		/** Called when the drawer size is changed by the user */
@@ -113,7 +113,7 @@ protected:
 	float MinDrawerSize = 0.f;
 	float MaxDrawerSize = 0.f;
 	float TargetDrawerSize = 0.f;
-	FVector2D ShadowOffset = FVector2D::ZeroVector;
+	FVector2f ShadowOffset = FVector2f::ZeroVector;
 	float ExpanderHandleSize = 5.f;
 
 	FOnSidebarDrawerTargetSizeChanged OnDrawerSizeChanged;
diff --git a/Engine/Source/Developer/ToolWidgets/Public/Dialog/SCustomDialog.h b/Engine/Source/Developer/ToolWidgets/Public/Dialog/SCustomDialog.h
--- a/Engine/Source/Developer/ToolWidgets/Public/Dialog/SCustomDialog.h
+++ b/Engine/Source/Developer/ToolWidgets/Public/Dialog/SCustomDialog.h
@@ -67,7 +67,7 @@ public:
 		, _RootPadding(FMargin(4.f))
 		, _ButtonAreaPadding(FMargin(20.f, 16.f, 0.f, 0.f))
 		, _UseScrollBox(true)
-		, _ScrollBoxMaxHeight(300)
+		, _ScrollBoxMaxHeight(300.f)
 		, _HAlignContent(HAlign_Left)
 		, _VAlignContent(VAlign_Center)
 	{
@@ -132,7 +132,7 @@ public:
 		SLATE_ARGUMENT(bool, UseScrollBox)
 
 		/** Max height for the scroll box (default: 300) */
-		SLATE_ARGUMENT(int32, ScrollBoxMaxHeight)
+		SLATE_ARGUMENT(float, ScrollBoxMaxHeight)
 
 	
 		/** HAlign to use for Content slot (default: HAlign_Left) */
diff --git a/Engine/Source/Developer/ToolWidgets/Public/Dialog/SMessageDialog.h b/Engine/Source/Developer/ToolWidgets/Public/Dialog/SMessageDialog.h
--- a/Engine/Source/Developer/ToolWidgets/Public/Dialog/SMessageDialog.h
+++ b/Engine/Source/Developer/ToolWidgets/Public/Dialog/SMessageDialog.h
@@ -21,7 +21,7 @@ public:
 		, _DecoratorStyleSet(nullptr)
 		, _Icon(nullptr)
 		, _UseScrollBox(true)
-		, _ScrollBoxMaxHeight(300)
+		, _ScrollBoxMaxHeight(300.f)
 		, _WrapMessageAt(512.f)
 	{}
 
@@ -63,7 +63,7 @@ public:
 		SLATE_ARGUMENT(bool, UseScrollBox)
 
 		/** Max height for the scroll box (default: 300) */
-		SLATE_ARGUMENT(int32, ScrollBoxMaxHeight)
+		SLATE_ARGUMENT(float, ScrollBoxMaxHeight)
 
 		/** When to wrap the message text (default: 512) */
 		SLATE_ATTRIBUTE(float, WrapMessageAt)
diff --git a/Engine/Source/Developer/ToolWidgets/Public/SSegmentedTimelineView.h b/Engine/Source/Developer/ToolWidgets/Public/SSegmentedTimelineView.h
--- a/Engine/Source/Developer/ToolWidgets/Public/SSegmentedTimelineView.h
+++ b/Engine/Source/Developer/ToolWidgets/Public/SSegmentedTimelineView.h
@@ -25,8 +25,8 @@ public:
 	};
 
 	SLATE_BEGIN_ARGS(SSegmentedTimelineView)
-		: _ViewRange(TRange<double>(0,10))
-		, _DesiredSize(FVector2D(100.f,20.f))
+		: _ViewRange(TRange<double>(0.0,10.0))
+		, _DesiredSize(FVector2f(100.f,20.f))
 	{}
     	/** View time range */
     	SLATE_ATTRIBUTE(TRange<double>, ViewRange);
@@ -35,7 +35,7 @@ public:
 		SLATE_ATTRIBUTE(TSharedPtr<FSegmentData>, SegmentData);
 
 		/** Desired widget size */
-		SLATE_ATTRIBUTE(FVector2D, DesiredSize);
+		SLATE_ATTRIBUTE(FVector2f, DesiredSize);
 	
 		/** Fill Color */
 		SLATE_ATTRIBUTE(FLinearColor, FillColor);
@@ -59,7 +59,7 @@ protected:
 	
 	TAttribute<TRange<double>> ViewRange;
 	TAttribute<TSharedPtr<FSegmentData>> SegmentData;
-	TAttribute<FVector2D> DesiredSize;
+	TAttribute<FVector2f> DesiredSize;
 	TAttribute<FLinearColor> FillColor;
 	
 	const FSlateBrush* WhiteBrush;
diff --git a/Engine/Source/Developer/ToolWidgets/Public/SSimpleTimeSlider.h b/Engine/Source/Developer/ToolWidgets/Public/SSimpleTimeSlider.h
--- a/Engine/Source/Developer/ToolWidgets/Public/SSimpleTimeSlider.h
+++ b/Engine/Source/Developer/ToolWidgets/Public/SSimpleTimeSlider.h
@@ -20,15 +20,15 @@ public:
 
 	SLATE_BEGIN_ARGS(SSimpleTimeSlider)
 		: _MirrorLabels( false )
-		, _ScrubPosition(0)
-		, _ViewRange(TRange<double>(0,10))
-		, _ClampRange(TRange<double>(0,10))
+		, _ScrubPosition(0.0)
+		, _ViewRange(TRange<double>(0.0,10.0))
+		, _ClampRange(TRange<double>(0.0,10.0))
 		, _AllowZoom (true)
 		, _AllowPan (true)
-		, _CursorSize(0)
+		, _CursorSize(0.f)
 		, _ClampRangeHighlightColor(FLinearColor(0.05f,0.05f,0.05f,1.0f))
-		, _ClampRangeHighlightSize(1.0)
-		, _DesiredSize(100,22)
+		, _ClampRangeHighlightSize(1.0f)
+		, _DesiredSize(FVector2f(100.f,22.f))
 	{}
 		/* If we should mirror the labels on the timeline */
 		SLATE_ATTRIBUTE( bool, MirrorLabels )
@@ -58,7 +58,7 @@ public:
 		SLATE_ATTRIBUTE(float, ClampRangeHighlightSize);
 
 		/* Widget desired size */
-		SLATE_ARGUMENT(FVector2d, DesiredSize);
+		SLATE_ARGUMENT(FVector2f, DesiredSize);
     
 	
 		/** Called when the scrub position changes */
@@ -91,32 +91,32 @@ public:
 	/** Utility struct for converting between scrub range space and local/absolute screen space */
 	struct FScrubRangeToScreen
 	{
-		FVector2D WidgetSize;
+		FVector2f WidgetSize;
 	
 		TRange<double> ViewInput;
 		double ViewInputRange;
 		float PixelsPerInput;
 	
-		FScrubRangeToScreen(TRange<double> InViewInput, const FVector2D& InWidgetSize )
+		FScrubRangeToScreen(TRange<double> InViewInput, const UE::Slate::FDeprecateVector2DParameter& InWidgetSize )
 		{
 			WidgetSize = InWidgetSize;
 	
 			ViewInput = InViewInput;
 			ViewInputRange = ViewInput.Size<double>();
-			PixelsPerInput = static_cast<float>(ViewInputRange > 0 ? (WidgetSize.X / ViewInputRange) : 0);
+			PixelsPerInput = ViewInputRange > 0.0 ? static_cast<float>(static_cast<double>(WidgetSize.X) / ViewInputRange) : 0.f;
 		}
 	
 		/** Local Widget Space -> Curve Input domain. */
-		float LocalXToInput(const float ScreenX) const
+		double LocalXToInput(const float ScreenX) const
 		{
 			const float LocalX = ScreenX;
-			return (LocalX/PixelsPerInput) + static_cast<float>(ViewInput.GetLowerBoundValue());
+			return static_cast<double>(LocalX/PixelsPerInput) + ViewInput.GetLowerBoundValue();
 		}
 	
 		/** Curve Input domain -> local Widget Space */
-		float InputToLocalX(const float Input) const
+		float InputToLocalX(const double Input) const
 		{
-			return (Input - static_cast<float>(ViewInput.GetLowerBoundValue())) * PixelsPerInput;
+			return static_cast<float>(Input - ViewInput.GetLowerBoundValue()) * PixelsPerInput;
 		}
 	};
 
@@ -163,7 +163,7 @@ protected:
 	 */
 	void DrawTicks( FSlateWindowElementList& OutDrawElements, const FScrubRangeToScreen& RangeToScreen, FDrawTickArgs& InArgs ) const;
 
-	float GetTimeAtCursorPosition(const FGeometry& MyGeometry, const FPointerEvent& MouseEvent) const;
+	double GetTimeAtCursorPosition(const FGeometry& MyGeometry, const FPointerEvent& MouseEvent) const;
 	
 	/**
 	 * Call this method when the user's interaction has changed the scrub position
@@ -171,7 +171,7 @@ protected:
 	 * @param NewValue				Value resulting from the user's interaction
 	 * @param bIsScrubbing			True if done via scrubbing, false if just releasing scrubbing
 	 */
-	void CommitScrubPosition( float NewValue, bool bIsScrubbing );
+	void CommitScrubPosition( double NewValue, bool bIsScrubbing );
 
 	TAttribute<double> ScrubPosition;
 	TAttribute<TRange<double>> ViewRange;
@@ -201,8 +201,8 @@ protected:
 	bool bPanning;
 	/***/
 	TSharedPtr<SScrollBar> Scrollbar;
-	FVector2D SoftwareCursorPosition;
-	FVector2D DesiredSize;
+	FVector2f SoftwareCursorPosition;
+	FVector2f DesiredSize;
 
 	
 };
diff --git a/Engine/Source/Developer/ToolWidgets/Public/Sidebar/SSidebarContainer.h b/Engine/Source/Developer/ToolWidgets/Public/Sidebar/SSidebarContainer.h
--- a/Engine/Source/Developer/ToolWidgets/Public/Sidebar/SSidebarContainer.h
+++ b/Engine/Source/Developer/ToolWidgets/Public/Sidebar/SSidebarContainer.h
@@ -1,4 +1,4 @@
-﻿// Copyright Epic Games, Inc. All Rights Reserved.
+// Copyright Epic Games, Inc. All Rights Reserved.
 
 #pragma once
 
@@ -33,7 +33,7 @@ public:
 	EOrientation GetSplitterOrientation() const;
 	ESidebarTabLocation GetTabLocation() const;
 	float GetCurrentDrawerSize() const;
-	FVector2D GetOverlaySize() const;
+	UE::Slate::FDeprecateVector2DResult GetOverlaySize() const;
 
 	void CloseAllDrawerWidgets(const bool bInAnimate);
 
diff --git a/Engine/Source/Developer/ToolWidgets/ToolWidgets.Build.cs b/Engine/Source/Developer/ToolWidgets/ToolWidgets.Build.cs
--- a/Engine/Source/Developer/ToolWidgets/ToolWidgets.Build.cs
+++ b/Engine/Source/Developer/ToolWidgets/ToolWidgets.Build.cs
@@ -6,6 +6,9 @@ public class ToolWidgets : ModuleRules
 {
 	public ToolWidgets(ReadOnlyTargetRules Target) : base(Target)
 	{
+		// Enable truncation warnings in this module
+		CppCompileWarningSettings.UnsafeTypeCastWarningLevel = WarningLevel.Warning;
+
 		/** NOTE: THIS MODULE SHOULD NOT EVER DEPEND ON UNREALED. 
 		 * If you are adding a reusable widget that depends on UnrealEd, add it to EditorWidgets instead
 		 */
