[ToolWidgets] Sidebar: Fix tab button not being removed when sidebar drawer is unregistered.

diff --git a/Engine/Source/Developer/ToolWidgets/Private/Sidebar/SSidebar.cpp b/Engine/Source/Developer/ToolWidgets/Private/Sidebar/SSidebar.cpp
--- a/Engine/Source/Developer/ToolWidgets/Private/Sidebar/SSidebar.cpp
+++ b/Engine/Source/Developer/ToolWidgets/Private/Sidebar/SSidebar.cpp
@@ -376,6 +376,11 @@ void SSidebar::RemoveDrawer(const TSharedRef<FSidebarDrawer>& InDrawer)
 		ContainerWidget->RemoveDrawerOverlaySlot(InDrawer, false);
 	}
 
+	if (InDrawer->ButtonWidget.IsValid())
+	{
+		TabButtonContainer->RemoveSlot(InDrawer->ButtonWidget.ToSharedRef());
+	}
+
 	const int32 IndexToRemove = Drawers.IndexOfByPredicate(
 		[DrawerId](const TSharedRef<FSidebarDrawer>& InDrawer)
 		{
