SCustomDialog can now control the enabled state and tooltip of each button

diff --git a/Engine/Source/Developer/ToolWidgets/Private/Dialog/SCustomDialog.cpp b/Engine/Source/Developer/ToolWidgets/Private/Dialog/SCustomDialog.cpp
--- a/Engine/Source/Developer/ToolWidgets/Private/Dialog/SCustomDialog.cpp
+++ b/Engine/Source/Developer/ToolWidgets/Private/Dialog/SCustomDialog.cpp
@@ -162,6 +162,8 @@ TSharedRef<SWidget> SCustomDialog::CreateButtonBox(const FArguments& InArgs)
 				&FAppStyle::Get().GetWidgetStyle< FButtonStyle >("Button");
 		
 		TSharedRef<SButton> ButtonWidget = SNew(SButton)
+			.IsEnabled(Button.ButtonIsEnabled)
+			.ToolTipText(Button.ButtonToolTipText)
 			.OnClicked(FOnClicked::CreateSP(this, &SCustomDialog::OnButtonClicked, Button.OnClicked, ButtonIndex))
 			.ButtonStyle(ButtonStyle)
 			[
diff --git a/Engine/Source/Developer/ToolWidgets/Public/Dialog/SCustomDialog.h b/Engine/Source/Developer/ToolWidgets/Public/Dialog/SCustomDialog.h
--- a/Engine/Source/Developer/ToolWidgets/Public/Dialog/SCustomDialog.h
+++ b/Engine/Source/Developer/ToolWidgets/Public/Dialog/SCustomDialog.h
@@ -34,11 +34,18 @@ public:
 	
 	struct FButton
 	{
-		FButton(const FText& InButtonText, const FSimpleDelegate& InOnClicked = FSimpleDelegate())
-			: ButtonText(InButtonText)
-			, OnClicked(InOnClicked)
+		FButton(TAttribute<FText> InButtonText, FSimpleDelegate InOnClicked = FSimpleDelegate())
+			: ButtonText(MoveTemp(InButtonText))
+			, ButtonIsEnabled(true)
+			, OnClicked(MoveTemp(InOnClicked))
 		{}
 
+		/** Set whether this button enabled */
+		FButton& SetIsEnabled(TAttribute<bool> InButtonIsEnabled) { ButtonIsEnabled = MoveTemp(InButtonIsEnabled); return *this; }
+
+		/** Set the tooltip text to use for this button */
+		FButton& SetToolTipText(TAttribute<FText> InButtonToolTipText) { ButtonToolTipText = MoveTemp(InButtonToolTipText); return *this; }
+
 		/** Primary buttons are highlighted */
 		FButton& SetPrimary(bool bValue) { bIsPrimary = bValue; return *this; }
 
@@ -49,11 +56,13 @@ public:
 		 * Whether to focus this button. Focus rules:
 		 *  1: If a button has SetFocus(), use the last one
 		 *  2: If a button is marked primary, use the last one
-		 *  3: Otherwise, ise the last button
+		 *  3: Otherwise, use the last button
 		 */
 		FButton& SetFocus() { bShouldFocus = true; return *this; }
 		
-		FText ButtonText;
+		TAttribute<FText> ButtonText;
+		TAttribute<FText> ButtonToolTipText;
+		TAttribute<bool> ButtonIsEnabled;
 		FSimpleDelegate OnClicked;
 		bool bIsPrimary = false;
 		bool bShouldFocus = false;
