Safeguard against crash reporter callstack

#tests Couldn't repro, but the callstack shows the window being null and the Pin() call is documented as potentially returning a null object so it feels safe. Also, the calling code probably wouldn't crash afterward as a result of skipping this part.
#rnx

diff --git a/Engine/Source/Developer/ToolWidgets/Private/SWidgetDrawer.cpp b/Engine/Source/Developer/ToolWidgets/Private/SWidgetDrawer.cpp
--- a/Engine/Source/Developer/ToolWidgets/Private/SWidgetDrawer.cpp
+++ b/Engine/Source/Developer/ToolWidgets/Private/SWidgetDrawer.cpp
@@ -553,11 +553,12 @@ FReply SWidgetDrawer::OnDrawerButtonClicked(const FName DrawerId)
 
 void SWidgetDrawer::OnDrawerHeightChanged(float TargetHeight)
 {
-	TSharedPtr<SWindow> MyWindow = OpenedDrawer.WindowWithOverlayContent.Pin();
-
-	// Save the height has a percentage of the screen
-	const float TargetDrawerHeightPct = TargetHeight / (MyWindow->GetSizeInScreen().Y / MyWindow->GetDPIScaleFactor());
-	GConfig->SetFloat(TEXT("DrawerSizes"), *(GetSerializableName() + TEXT(".") + OpenedDrawer.DrawerId.ToString()), TargetDrawerHeightPct, GEditorSettingsIni);
+	if (TSharedPtr<SWindow> MyWindow = OpenedDrawer.WindowWithOverlayContent.Pin())
+	{
+		// Save the height has a percentage of the screen
+		const float TargetDrawerHeightPct = TargetHeight / (MyWindow->GetSizeInScreen().Y / MyWindow->GetDPIScaleFactor());
+		GConfig->SetFloat(TEXT("DrawerSizes"), *(GetSerializableName() + TEXT(".") + OpenedDrawer.DrawerId.ToString()), TargetDrawerHeightPct, GEditorSettingsIni);
+	}
 }
 
 void SWidgetDrawer::CloseDrawerImmediatelyInternal(const FOpenDrawerData& Data)
