[ToolWidgets] Sidebar: Close previous open drawers when a new drawer is opened. (Used by Sequencer and Motion Design sidebars)

diff --git a/Engine/Source/Developer/ToolWidgets/Private/Sidebar/SSidebarContainer.cpp b/Engine/Source/Developer/ToolWidgets/Private/Sidebar/SSidebarContainer.cpp
--- a/Engine/Source/Developer/ToolWidgets/Private/Sidebar/SSidebarContainer.cpp
+++ b/Engine/Source/Developer/ToolWidgets/Private/Sidebar/SSidebarContainer.cpp
@@ -455,6 +455,11 @@ void SSidebarContainer::OpenDrawer_Internal(const TSharedRef<FSidebarDrawer>& In
 		return;
 	}
 
+	for (const TSharedRef<FSidebarDrawer>& Drawer : SidebarWidget->GetAllDrawers())
+	{
+		CloseDrawer_Internal(Drawer, false);
+	}
+
 	PendingTabToOpen.Reset();
 	bAnimatePendingTabOpen = false;
 
@@ -462,13 +467,7 @@ void SSidebarContainer::OpenDrawer_Internal(const TSharedRef<FSidebarDrawer>& In
 	AddDrawerOverlaySlot(InDrawer);
 
 	InDrawer->DrawerWidget->Open(bInAnimate);
-
-	for (const TSharedRef<FSidebarDrawer>& Drawer : SidebarWidget->GetAllDrawers())
-	{
-		Drawer->bIsOpen = false;
-	}
 	InDrawer->bIsOpen = true;
-
 	InDrawer->DrawerOpenedDelegate.ExecuteIfBound(InDrawer->GetUniqueId());
 
 	UpdateDrawerTabAppearance();
