[Backout] - CL42038394
#rnx
Original CL Desc
-----------------------------------------------------------------
[ToolWidgets] Sidebar: Close previous open drawers when a new drawer is opened. (Used by Sequencer and Motion Design sidebars)

diff --git a/Engine/Source/Developer/ToolWidgets/Private/Sidebar/SSidebarContainer.cpp b/Engine/Source/Developer/ToolWidgets/Private/Sidebar/SSidebarContainer.cpp
--- a/Engine/Source/Developer/ToolWidgets/Private/Sidebar/SSidebarContainer.cpp
+++ b/Engine/Source/Developer/ToolWidgets/Private/Sidebar/SSidebarContainer.cpp
@@ -455,11 +455,6 @@ void SSidebarContainer::OpenDrawer_Internal(const TSharedRef<FSidebarDrawer>& In
 		return;
 	}
 
-	for (const TSharedRef<FSidebarDrawer>& Drawer : SidebarWidget->GetAllDrawers())
-	{
-		CloseDrawer_Internal(Drawer, false);
-	}
-
 	PendingTabToOpen.Reset();
 	bAnimatePendingTabOpen = false;
 
@@ -467,7 +462,13 @@ void SSidebarContainer::OpenDrawer_Internal(const TSharedRef<FSidebarDrawer>& In
 	AddDrawerOverlaySlot(InDrawer);
 
 	InDrawer->DrawerWidget->Open(bInAnimate);
+
+	for (const TSharedRef<FSidebarDrawer>& Drawer : SidebarWidget->GetAllDrawers())
+	{
+		Drawer->bIsOpen = false;
+	}
 	InDrawer->bIsOpen = true;
+
 	InDrawer->DrawerOpenedDelegate.ExecuteIfBound(InDrawer->GetUniqueId());
 
 	UpdateDrawerTabAppearance();
