[ToolWidgets] Sidebar: Close previous open drawers when a new drawer is opened.
(Used by Sequencer and Motion Design sidebars)

diff --git a/Engine/Source/Developer/ToolWidgets/Private/Sidebar/SSidebarContainer.cpp b/Engine/Source/Developer/ToolWidgets/Private/Sidebar/SSidebarContainer.cpp
--- a/Engine/Source/Developer/ToolWidgets/Private/Sidebar/SSidebarContainer.cpp
+++ b/Engine/Source/Developer/ToolWidgets/Private/Sidebar/SSidebarContainer.cpp
@@ -455,6 +455,11 @@ void SSidebarContainer::OpenDrawer_Internal(const TSharedRef<FSidebarDrawer>& In
 		return;
 	}
 
+	for (const TSharedRef<FSidebarDrawer>& Drawer : SidebarWidget->GetAllDrawers())
+	{
+		CloseDrawer_Internal(Drawer, false, false);
+	}
+
 	PendingTabToOpen.Reset();
 	bAnimatePendingTabOpen = false;
 
@@ -462,13 +467,7 @@ void SSidebarContainer::OpenDrawer_Internal(const TSharedRef<FSidebarDrawer>& In
 	AddDrawerOverlaySlot(InDrawer);
 
 	InDrawer->DrawerWidget->Open(bInAnimate);
-
-	for (const TSharedRef<FSidebarDrawer>& Drawer : SidebarWidget->GetAllDrawers())
-	{
-		Drawer->bIsOpen = false;
-	}
 	InDrawer->bIsOpen = true;
-
 	InDrawer->DrawerOpenedDelegate.ExecuteIfBound(InDrawer->GetUniqueId());
 
 	UpdateDrawerTabAppearance();
@@ -478,22 +477,29 @@ void SSidebarContainer::OpenDrawer_Internal(const TSharedRef<FSidebarDrawer>& In
 	FSlateApplication::Get().SetKeyboardFocus(InDrawer->DrawerWidget);
 }
 
-void SSidebarContainer::CloseDrawer_Internal(const TSharedRef<FSidebarDrawer>& InDrawer, const bool bInAnimate)
+void SSidebarContainer::CloseDrawer_Internal(const TSharedRef<FSidebarDrawer>& InDrawer, const bool bInAnimate
+	, const bool bInSummonPinnedTabIfNothingOpened)
 {
 	const TSharedPtr<SSidebarDrawer> FoundDrawerWidget = FindOpenDrawerWidget(InDrawer);
-	if (OpenDrawerWidgets.Contains(FoundDrawerWidget) && !ClosingDrawerWidgets.Contains(InDrawer->DrawerWidget))
+	if (!FoundDrawerWidget.IsValid()
+		|| !OpenDrawerWidgets.Contains(FoundDrawerWidget)
+		|| ClosingDrawerWidgets.Contains(InDrawer->DrawerWidget))
 	{
-		const TSharedRef<SSidebarDrawer> FoundDrawerWidgetRef = FoundDrawerWidget.ToSharedRef();
+		return;
+	}
 
-		InDrawer->bIsOpen = false;
+	InDrawer->bIsOpen = false;
 
-		RemoveDrawerOverlaySlot(InDrawer, bInAnimate);
+	RemoveDrawerOverlaySlot(InDrawer, bInAnimate);
 
-		FoundDrawerWidgetRef->Close(bInAnimate);
-	}
+	FoundDrawerWidget->Close(bInAnimate);
 
-	SummonPinnedTabIfNothingOpened();
 	UpdateDrawerTabAppearance();
+
+	if (bInSummonPinnedTabIfNothingOpened)
+	{
+		SummonPinnedTabIfNothingOpened();
+	}
 }
 
 void SSidebarContainer::SummonPinnedTabIfNothingOpened()
diff --git a/Engine/Source/Developer/ToolWidgets/Public/Sidebar/SSidebarContainer.h b/Engine/Source/Developer/ToolWidgets/Public/Sidebar/SSidebarContainer.h
--- a/Engine/Source/Developer/ToolWidgets/Public/Sidebar/SSidebarContainer.h
+++ b/Engine/Source/Developer/ToolWidgets/Public/Sidebar/SSidebarContainer.h
@@ -77,7 +77,8 @@ protected:
 
 	UE_API void OpenDrawerNextFrame(const TSharedRef<FSidebarDrawer>& InDrawer, const bool bInAnimate = true);
 	UE_API void OpenDrawer_Internal(const TSharedRef<FSidebarDrawer>& InDrawer, const bool bInAnimate = true);
-	UE_API void CloseDrawer_Internal(const TSharedRef<FSidebarDrawer>& InDrawer, const bool bInAnimate = true);
+	UE_API void CloseDrawer_Internal(const TSharedRef<FSidebarDrawer>& InDrawer, const bool bInAnimate = true
+		, const bool bInSummonPinnedTabIfNothingOpened = true);
 
 	UE_API void DockDrawer_Internal(const TSharedRef<FSidebarDrawer>& InDrawer);
 	UE_API void UndockDrawer_Internal(const TSharedRef<FSidebarDrawer>& InDrawer);
