Tool Widgets: Consolidate SDraggableBox to ToolWidgets.

SDraggableBox was first a widget from the modelling framework, was then extended to SBoundDraggableBox in non-licensee code, and finally we are surfacing it to public engine interface with improvements made.
This version of SDraggableBox anchors the widget to the corners: e.g. when a user drag it from bottom-left to bottom-right corner and resizes the viewport, the widget snaps to the new corner.

diff --git a/Engine/Plugins/Experimental/MeshModelingToolsetExp/Source/ModelingEditorUI/Private/STransformGizmoNumericalUIOverlay.cpp b/Engine/Plugins/Experimental/MeshModelingToolsetExp/Source/ModelingEditorUI/Private/STransformGizmoNumericalUIOverlay.cpp
--- a/Engine/Plugins/Experimental/MeshModelingToolsetExp/Source/ModelingEditorUI/Private/STransformGizmoNumericalUIOverlay.cpp
+++ b/Engine/Plugins/Experimental/MeshModelingToolsetExp/Source/ModelingEditorUI/Private/STransformGizmoNumericalUIOverlay.cpp
@@ -10,9 +10,9 @@
 #include "Framework/MultiBox/MultiBoxBuilder.h" // FMenuBuilder
 #include "InteractiveToolsContext.h"
 #include "Misc/AxisDisplayInfo.h"
-#include "ModelingWidgets/SDraggableBox.h" // SDraggableBoxOverlay
 #include "ModelingWidgets/STickDelegateWidget.h"
 #include "Transforms/TransformGizmoDataBinder.h"
+#include "Overlay/SDraggableBoxOverlay.h"
 #include "Widgets/Input/SButton.h"
 #include "Widgets/Layout/SBox.h"
 #include "Widgets/SBoxPanel.h" // SHorizontalBox
@@ -49,8 +49,9 @@ void STransformGizmoNumericalUIOverlay::Construct(const FArguments& InArgs)
 
 	ChildSlot
 	[
-		SAssignNew(DraggableBoxOverlay, SDraggableBoxOverlay)
-		.bPositionRelativeToBottom(bPositionRelativeToBottom)
+		SAssignNew(DraggableBoxOverlay, UE::ToolWidgets::SDraggableBoxOverlay)
+		.InitialAlignmentOffset({ DefaultLeftPadding, DefaultVerticalPadding })
+		.VAlign(bPositionRelativeToBottom ? VAlign_Bottom : VAlign_Top)
 		.Visibility_Lambda([this]()
 		{
 			return bIsEnabled && DataBinder->HasVisibleGizmo() ? EVisibility::SelfHitTestInvisible : EVisibility::Collapsed;
@@ -69,8 +70,6 @@ void STransformGizmoNumericalUIOverlay::Construct(const FArguments& InArgs)
 
 	// Don't block things under the overlay.
 	SetVisibility(EVisibility::SelfHitTestInvisible);
-
-	ResetPositionInViewport();
 }
 
 STransformGizmoNumericalUIOverlay::~STransformGizmoNumericalUIOverlay()
@@ -468,7 +467,9 @@ void STransformGizmoNumericalUIOverlay::ResetPositionInViewport()
 {
 	if (DraggableBoxOverlay.IsValid())
 	{
-		DraggableBoxOverlay->SetBoxPosition(DefaultLeftPadding, DefaultVerticalPadding);
+		DraggableBoxOverlay->SetBoxHorizontalAlignment(HAlign_Left);
+		DraggableBoxOverlay->SetBoxVerticalAlignment(bPositionRelativeToBottom ? VAlign_Bottom : VAlign_Top);
+		DraggableBoxOverlay->SetBoxAlignmentOffset({ DefaultLeftPadding, DefaultVerticalPadding });
 	}
 }
 
diff --git a/Engine/Plugins/Experimental/MeshModelingToolsetExp/Source/ModelingEditorUI/Public/STransformGizmoNumericalUIOverlay.h b/Engine/Plugins/Experimental/MeshModelingToolsetExp/Source/ModelingEditorUI/Public/STransformGizmoNumericalUIOverlay.h
--- a/Engine/Plugins/Experimental/MeshModelingToolsetExp/Source/ModelingEditorUI/Public/STransformGizmoNumericalUIOverlay.h
+++ b/Engine/Plugins/Experimental/MeshModelingToolsetExp/Source/ModelingEditorUI/Public/STransformGizmoNumericalUIOverlay.h
@@ -11,7 +11,7 @@
 
 class FTransformGizmoDataBinder;
 class SBorder;
-class SDraggableBoxOverlay;
+namespace UE::ToolWidgets { class SDraggableBoxOverlay; }
 class SWidget;
 class UInteractiveToolsContext;
 
@@ -110,7 +110,7 @@ protected:
 	// Holds the functionality that binds the displayed values to the gizmo.
 	TSharedPtr<FTransformGizmoDataBinder> DataBinder;
 
-	TSharedPtr<SDraggableBoxOverlay> DraggableBoxOverlay;
+	TSharedPtr<UE::ToolWidgets::SDraggableBoxOverlay> DraggableBoxOverlay;
 
 	// The part of the UI that changes when a different gizmo is chosen
 	TSharedPtr<SBorder> WidgetContents;
diff --git a/Engine/Plugins/Experimental/MeshModelingToolsetExp/Source/ModelingUI/Private/ModelingWidgets/SDraggableBox.cpp b/Engine/Plugins/Experimental/MeshModelingToolsetExp/Source/ModelingUI/Private/ModelingWidgets/SDraggableBox.cpp
--- a/Engine/Plugins/Experimental/MeshModelingToolsetExp/Source/ModelingUI/Private/ModelingWidgets/SDraggableBox.cpp
+++ b/Engine/Plugins/Experimental/MeshModelingToolsetExp/Source/ModelingUI/Private/ModelingWidgets/SDraggableBox.cpp
@@ -10,7 +10,7 @@
 #include "Widgets/SBoxPanel.h" // SVerticalBox
 #include "Widgets/SWindow.h"
 
-
+PRAGMA_DISABLE_DEPRECATION_WARNINGS
 void SDraggableBoxOverlay::Construct(const FArguments& InArgs)
 {
 	// Depending on whether we're positioning relative to the bottom or top, we'll update the margin of the
@@ -68,6 +68,7 @@ void SDraggableBoxOverlay::Construct(const FArguments& InArgs)
 	]
 	;
 }
+PRAGMA_ENABLE_DEPRECATION_WARNINGS
 
 void SDraggableBoxOverlay::SetBoxPosition(float HorizontalPosition, float VerticalPosition)
 {
@@ -111,6 +112,7 @@ FReply SDraggableBox::OnMouseButtonDown(const FGeometry& MyGeometry, const FPoin
 	return FReply::Handled().DetectDrag(SharedThis(this), EKeys::LeftMouseButton);
 };
 
+PRAGMA_DISABLE_DEPRECATION_WARNINGS
 FReply SDraggableBox::OnDragDetected(const FGeometry& MyGeometry, const FPointerEvent& MouseEvent)
 {
 	// The drag/drop operation moves a transparent copy of our contents, so we should hide ourselves to make it
@@ -134,6 +136,7 @@ FReply SDraggableBox::OnDragDetected(const FGeometry& MyGeometry, const FPointer
 
 	return FReply::Handled().BeginDragDrop(DragDropOperation);
 }
+PRAGMA_ENABLE_DEPRECATION_WARNINGS
 
 // Drag operation:
 
@@ -158,13 +161,16 @@ void FDraggableBoxUIDragOperation::OnDragged(const FDragDropEvent& DragDropEvent
 	CursorDecoratorWindow->MoveWindowTo(TargetPosition);
 }
 
+PRAGMA_DISABLE_DEPRECATION_WARNINGS
 TSharedRef<FDraggableBoxUIDragOperation> FDraggableBoxUIDragOperation::New(const TSharedRef<SWidget>& InUIToBeDragged, 
 	const FVector2D InDecoratorOffset, const FVector2D& OwnerAreaSize, SDraggableBox::FOnDragComplete& OnDragComplete)
 {
 	const TSharedRef<FDraggableBoxUIDragOperation> Operation = MakeShareable(new FDraggableBoxUIDragOperation(InUIToBeDragged, InDecoratorOffset, OwnerAreaSize, OnDragComplete));
 	return Operation;
 }
+PRAGMA_ENABLE_DEPRECATION_WARNINGS
 
+PRAGMA_DISABLE_DEPRECATION_WARNINGS
 FDraggableBoxUIDragOperation::FDraggableBoxUIDragOperation(const TSharedRef<class SWidget>& InUIToBeDragged, 
 	const FVector2D InDecoratorOffset, const FVector2D& OwnerAreaSize, SDraggableBox::FOnDragComplete& InOnDragComplete)
 	: UIBeingDragged(InUIToBeDragged)
@@ -182,4 +188,5 @@ FDraggableBoxUIDragOperation::FDraggableBoxUIDragOperation(const TSharedRef<clas
 		InUIToBeDragged
 	);
 
-}
\ No newline at end of file
+}
+PRAGMA_ENABLE_DEPRECATION_WARNINGS
\ No newline at end of file
diff --git a/Engine/Plugins/Experimental/MeshModelingToolsetExp/Source/ModelingUI/Public/ModelingWidgets/SDraggableBox.h b/Engine/Plugins/Experimental/MeshModelingToolsetExp/Source/ModelingUI/Public/ModelingWidgets/SDraggableBox.h
--- a/Engine/Plugins/Experimental/MeshModelingToolsetExp/Source/ModelingUI/Public/ModelingWidgets/SDraggableBox.h
+++ b/Engine/Plugins/Experimental/MeshModelingToolsetExp/Source/ModelingUI/Public/ModelingWidgets/SDraggableBox.h
@@ -11,7 +11,7 @@
  * Class that can be used to place a draggable box into a viewport or some other large widget as an
  * overlay. Just place the widget that you want to be draggable as the contents of SDraggableBoxOverlay.
  */
-class SDraggableBoxOverlay : public SCompoundWidget
+class UE_DEPRECATED(5.6, "Use UE::ToolWidgets::SDraggableBoxOverlay instead.") SDraggableBoxOverlay : public SCompoundWidget
 {
 public:
 
@@ -52,7 +52,7 @@ protected:
  * which will put its contents into a draggable box and properly handle the dragging without the
  * user having to set it up.
  */
-class SDraggableBox : public SCompoundWidget
+class UE_DEPRECATED(5.6, "SDraggableBox is now private. Use UE::ToolWidgets::SDraggableBoxOverlay instead.") SDraggableBox : public SCompoundWidget
 {
 public:
 
@@ -85,7 +85,7 @@ protected:
  * OnDragComplete delegate, it just triggers on any drop. This makes it possible to use it in any
  * viewports, not just the level editor.
  */
-class FDraggableBoxUIDragOperation : public FDragDropOperation
+class UE_DEPRECATED(5.6, "FDraggableBoxUIDragOperation is now private. Use UE::ToolWidgets::SDraggableBoxOverlay instead.") FDraggableBoxUIDragOperation : public FDragDropOperation
 {
 public:
 
@@ -105,15 +105,22 @@ public:
 	 *
 	 * @return a new FDockingDragOperation
 	 */
+	PRAGMA_DISABLE_DEPRECATION_WARNINGS
 	MODELINGUI_API static TSharedRef<FDraggableBoxUIDragOperation> New(const TSharedRef<class SWidget>& InUIToBeDragged, const FVector2D InDecoratorOffset, 
 		const FVector2D& OwnerAreaSize, SDraggableBox::FOnDragComplete& OnDragComplete);
+	PRAGMA_ENABLE_DEPRECATION_WARNINGS
 
+	PRAGMA_DISABLE_DEPRECATION_WARNINGS
 	virtual ~FDraggableBoxUIDragOperation() {}
+	PRAGMA_ENABLE_DEPRECATION_WARNINGS
 
 protected:
+	
 	/** The constructor is protected, so that this class can only be instantiated as a shared pointer. */
+	PRAGMA_DISABLE_DEPRECATION_WARNINGS
 	MODELINGUI_API FDraggableBoxUIDragOperation(const TSharedRef<class SWidget>& InUIToBeDragged, const FVector2D InDecoratorOffsetFromCursor, 
 		const FVector2D& OwnerAreaSize, SDraggableBox::FOnDragComplete& OnDragCompletee);
+	PRAGMA_ENABLE_DEPRECATION_WARNINGS
 
 	/**
 	 * Shared pointer to our contents, which are a weak pointer.
@@ -125,6 +132,8 @@ protected:
 
 	// Size of rectangle being dragged
 	FVector2D LastContentSize;
-
+	
+	PRAGMA_DISABLE_DEPRECATION_WARNINGS
 	SDraggableBox::FOnDragComplete OnDragComplete;
+	PRAGMA_ENABLE_DEPRECATION_WARNINGS
 };
diff --git a/Engine/Plugins/Runtime/ScriptableToolsFramework/Source/ScriptableToolsFramework/Private/ScriptableInteractiveTool.cpp b/Engine/Plugins/Runtime/ScriptableToolsFramework/Source/ScriptableToolsFramework/Private/ScriptableInteractiveTool.cpp
--- a/Engine/Plugins/Runtime/ScriptableToolsFramework/Source/ScriptableToolsFramework/Private/ScriptableInteractiveTool.cpp
+++ b/Engine/Plugins/Runtime/ScriptableToolsFramework/Source/ScriptableToolsFramework/Private/ScriptableInteractiveTool.cpp
@@ -23,9 +23,9 @@
 
 #include "Blueprint/UserWidget.h"
 #include "ContextObjectStore.h"
+#include "Overlay/SDraggableBoxOverlay.h"
 #include "Utility/ScriptableToolContextObjects.h"
 #include "Slate/SObjectWidget.h"
-#include "ModelingWidgets/SDraggableBox.h"
 #include "Components/Widget.h"
 
 
@@ -1286,7 +1286,7 @@ void UScriptableInteractiveTool::SetOverlayWidget(UUserWidget* Widget, bool bMak
 			TSharedRef<SWidget> WrappedWidget = Widget->TakeDerivedWidget<ScriptableInteractiveToolLocals::SScriptableToolViewportOverlayDragContainer>([](UUserWidget* Widget, TSharedRef<SWidget> Content) {
 				return SNew(ScriptableInteractiveToolLocals::SScriptableToolViewportOverlayDragContainer, Widget)
 					[
-						SNew(SDraggableBoxOverlay)
+						SNew(UE::ToolWidgets::SDraggableBoxOverlay)
 						[
 							Content
 						]
diff --git a/Engine/Plugins/Runtime/ScriptableToolsFramework/Source/ScriptableToolsFramework/ScriptableToolsFramework.Build.cs b/Engine/Plugins/Runtime/ScriptableToolsFramework/Source/ScriptableToolsFramework/ScriptableToolsFramework.Build.cs
--- a/Engine/Plugins/Runtime/ScriptableToolsFramework/Source/ScriptableToolsFramework/ScriptableToolsFramework.Build.cs
+++ b/Engine/Plugins/Runtime/ScriptableToolsFramework/Source/ScriptableToolsFramework/ScriptableToolsFramework.Build.cs
@@ -35,7 +35,6 @@ public class ScriptableToolsFramework : ModuleRules
 				"InteractiveToolsFramework",
 				"ModelingComponents",
 				"UMG",
-				"ModelingUI"
 			}
 			);
 
@@ -45,7 +44,8 @@ public class ScriptableToolsFramework : ModuleRules
 			{
 				"Engine",
 				"Slate",
-				"SlateCore"
+				"SlateCore",
+				"ToolWidgets"
 			}
 			);
 
diff --git a/Engine/Source/Developer/ToolWidgets/Private/Overlay/SDraggableBox.cpp b/Engine/Source/Developer/ToolWidgets/Private/Overlay/SDraggableBox.cpp
new file mode 100644
--- /dev/null
+++ b/Engine/Source/Developer/ToolWidgets/Private/Overlay/SDraggableBox.cpp
@@ -0,0 +1,145 @@
+﻿// Copyright Epic Games, Inc. All Rights Reserved.
+
+#include "Overlay/SDraggableBox.h"
+
+#include "Framework/Application/SlateApplication.h"
+#include "Overlay/SDraggableBoxOverlay.h"
+#include "Input/DragAndDrop.h"
+#include "Styling/SlateTypes.h"
+#include "Widgets/Layout/SBorder.h"
+#include "Widgets/SBoxPanel.h"
+#include "Widgets/SWindow.h"
+
+namespace UE::ToolWidgets
+{
+namespace DraggableBoxDetail
+{
+/** A drag/drop operation used by SDraggableBox. */
+class FDraggableBoxUIDragOperation : public FDragDropOperation
+{
+public:
+	FDraggableBoxUIDragOperation(const TSharedRef<SDraggableBox> InDraggableBox,
+		const SDraggableBox::FDragInfo& InDragInfo)
+		: DraggableBoxWeak(InDraggableBox)
+		, DragInfo(InDragInfo)
+	{
+	}
+
+	virtual ~FDraggableBoxUIDragOperation() override = default;
+
+	//~ Begin FDragDropOperation
+	virtual void OnDragged(const FDragDropEvent& InDragDropEvent)
+	{
+		if (TSharedPtr<SDraggableBox> DraggableBox = DraggableBoxWeak.Pin())
+		{
+			DraggableBox->OnDragUpdate(InDragDropEvent, DragInfo, /* Dropped */ false);
+		}
+	}
+
+	virtual void OnDrop(bool bDropWasHandled, const FPointerEvent& InMouseEvent)
+	{
+		if (TSharedPtr<SDraggableBox> DraggableBox = DraggableBoxWeak.Pin())
+		{
+			DraggableBox->OnDragUpdate(InMouseEvent, DragInfo, /* Dropped */ true);
+		}
+	}
+	//~ End FDragDropOperation
+
+protected:
+	TWeakPtr<SDraggableBox> DraggableBoxWeak;
+	SDraggableBox::FDragInfo DragInfo;
+};
+}
+
+void SDraggableBox::Construct(const FArguments& InArgs, const TSharedRef<SDraggableBoxOverlay>& InDraggableOverlay)
+{
+	DraggableOverlayWeak = InDraggableOverlay;
+	InnerWidget = InArgs._Content.Widget;
+	IsDraggableAttr = InArgs._IsDraggable;
+	OnUserDraggedToNewPositionDelegate = InArgs._OnUserDraggedToNewPosition;
+
+	ChildSlot
+	[
+		InArgs._Content.Widget
+	];
+}
+
+void SDraggableBox::OnDragUpdate(const FPointerEvent& InMouseEvent, const FDragInfo& InDragInfo, bool bInDropped)
+{
+	const TSharedPtr<SDraggableBoxOverlay> DraggableOverlay = DraggableOverlayWeak.Pin();
+	if (!DraggableOverlay.IsValid() || !IsDraggableAttr.Get())
+	{
+		return;
+	}
+
+	const FGeometry& MyGeometry = DraggableOverlay->GetTickSpaceGeometry();
+	const FVector2f MouseOffset = (InMouseEvent.GetScreenSpacePosition() - InDragInfo.OriginalMousePosition)
+		* (MyGeometry.GetLocalSize() / MyGeometry.GetAbsoluteSize());
+
+	FVector2f NewAlignmentOffset = InDragInfo.OriginalAlignmentOffset;
+	switch (InDragInfo.OriginalHorizontalAlignment)
+	{
+	case EHorizontalAlignment::HAlign_Left:
+		NewAlignmentOffset.X += MouseOffset.X;
+		break;
+
+	case EHorizontalAlignment::HAlign_Right:
+		NewAlignmentOffset.X -= MouseOffset.X;
+		break;
+
+	default:
+		// Do nothing
+		break;
+	}
+
+	switch (InDragInfo.OriginalVerticalAlignment)
+	{
+	case EVerticalAlignment::VAlign_Top:
+		NewAlignmentOffset.Y += MouseOffset.Y;
+		break;
+
+	case EVerticalAlignment::VAlign_Bottom:
+		NewAlignmentOffset.Y -= MouseOffset.Y;
+		break;
+
+	default:
+		// Do nothing
+		break;
+	}
+
+	DraggableOverlay->SetBoxHorizontalAlignment(InDragInfo.OriginalHorizontalAlignment);
+	DraggableOverlay->SetBoxVerticalAlignment(InDragInfo.OriginalVerticalAlignment);
+	DraggableOverlay->SetBoxAlignmentOffset(NewAlignmentOffset);
+
+	if (bInDropped)
+	{
+		OnUserDraggedToNewPositionDelegate.ExecuteIfBound();
+	}
+}
+
+FReply SDraggableBox::OnMouseButtonDown(const FGeometry& InGeometry, const FPointerEvent& InMouseEvent)
+{
+	return FReply::Handled().DetectDrag(SharedThis(this), EKeys::LeftMouseButton);
+}
+
+FReply SDraggableBox::OnDragDetected(const FGeometry& InGeometry, const FPointerEvent& InMouseEvent)
+{
+	const TSharedPtr<SDraggableBoxOverlay> DraggableOverlay = DraggableOverlayWeak.Pin();
+	if (DraggableOverlay && IsDraggableAttr.Get())
+	{
+		const FDragInfo DragInfo
+		{
+			DraggableOverlay->GetBoxHorizontalAlignment(),
+			DraggableOverlay->GetBoxVerticalAlignment(),
+			DraggableOverlay->GetBoxAlignmentOffset(),
+			InMouseEvent.GetScreenSpacePosition()
+		};
+		using namespace DraggableBoxDetail;
+		const TSharedRef<FDraggableBoxUIDragOperation> DragDropOperation = MakeShared<FDraggableBoxUIDragOperation>(
+			SharedThis(this), DragInfo
+		);
+		return FReply::Handled().BeginDragDrop(DragDropOperation);
+	}
+	return FReply::Unhandled();
+}
+}
\ No newline at end of file
diff --git a/Engine/Source/Developer/ToolWidgets/Private/Overlay/SDraggableBox.h b/Engine/Source/Developer/ToolWidgets/Private/Overlay/SDraggableBox.h
new file mode 100644
--- /dev/null
+++ b/Engine/Source/Developer/ToolWidgets/Private/Overlay/SDraggableBox.h
@@ -0,0 +1,66 @@
+﻿// Copyright Epic Games, Inc. All Rights Reserved.
+
+#pragma once
+
+#include "Widgets/SCompoundWidget.h"
+#include "Widgets/SBoxPanel.h"
+#include "Widgets/DeclarativeSyntaxSupport.h"
+
+namespace UE::ToolWidgets
+{
+class SDraggableBoxOverlay;
+	
+/**
+ * A widget for the draggable box itself, which requires its parent to handle its positioning in
+ * response to the drag.
+ *
+ * Users probably shouldn't use this class directly; rather, they should use SDraggableBoxOverlay,
+ * which will put its contents into a draggable box and properly handle the dragging without the
+ * user having to set it up.
+ */
+class SDraggableBox : public SCompoundWidget
+{
+public:
+	SLATE_BEGIN_ARGS(SDraggableBox)
+		: _IsDraggable(true)
+	{}
+	SLATE_DEFAULT_SLOT(FArguments, Content)
+
+	/** Whether this widget should be allowed to be dragged. */
+	SLATE_ATTRIBUTE(bool, IsDraggable)
+	
+	/** Invoked when the user has finished dragging the widget to a new position. */
+	SLATE_EVENT(FSimpleDelegate, OnUserDraggedToNewPosition)
+SLATE_END_ARGS()
+
+void Construct(const FArguments& InArgs, const TSharedRef<SDraggableBoxOverlay>& InDraggableOverlay);
+
+	struct FDragInfo
+	{
+		EHorizontalAlignment OriginalHorizontalAlignment = HAlign_Left;
+		EVerticalAlignment OriginalVerticalAlignment = VAlign_Bottom;
+		FVector2f OriginalAlignmentOffset = FVector2f::ZeroVector;
+		FVector2f OriginalMousePosition;
+	};
+
+	void OnDragUpdate(const FPointerEvent& InMouseEvent, const FDragInfo& InDragInfo, bool bInDropped);
+
+	//~ Begin SWidget
+	FReply OnMouseButtonDown(const FGeometry& InGeometry, const FPointerEvent& InMouseEvent) override;
+	FReply OnDragDetected(const FGeometry& InGeometry, const FPointerEvent& InMouseEvent) override;
+	//~ End SWidget
+
+protected:
+
+	TWeakPtr<SDraggableBoxOverlay> DraggableOverlayWeak;
+	TSharedPtr<SWidget> InnerWidget;
+
+private:
+
+	/** Whether this widget should be allowed to be dragged. */
+	TAttribute<bool> IsDraggableAttr;
+
+	/** The content you want to be able to drag around in the parent widget. */
+	FSimpleDelegate OnUserDraggedToNewPositionDelegate;
+};
+}
\ No newline at end of file
diff --git a/Engine/Source/Developer/ToolWidgets/Private/Overlay/SDraggableBoxOverlay.cpp b/Engine/Source/Developer/ToolWidgets/Private/Overlay/SDraggableBoxOverlay.cpp
new file mode 100644
--- /dev/null
+++ b/Engine/Source/Developer/ToolWidgets/Private/Overlay/SDraggableBoxOverlay.cpp
@@ -0,0 +1,233 @@
+﻿// Copyright Epic Games, Inc. All Rights Reserved.
+
+#include "Overlay/SDraggableBoxOverlay.h"
+
+#include "Framework/Application/SlateApplication.h"
+#include "Overlay/DragBoxPosition.h"
+#include "Overlay/SDraggableBox.h"
+#include "Widgets/Layout/SBox.h"
+#include "Widgets/SWindow.h"
+
+namespace UE::ToolWidgets
+{
+namespace DraggableBoxDetail
+{
+	static constexpr float DraggableBorderArea = 3.f;
+}
+
+void SDraggableBoxOverlay::Construct(const FArguments& InArgs)
+{
+	HorizontalAlignment = InArgs._HAlign;
+	VerticalAlignment = InArgs._VAlign;
+
+	ChildSlot
+	[
+		SAssignNew(Container, SBox)
+		.HAlign(HorizontalAlignment)
+		.VAlign(VerticalAlignment)
+		[
+			SAssignNew(DraggableBox, SDraggableBox, SharedThis(this))
+			.IsDraggable(InArgs._IsDraggable)
+			.OnUserDraggedToNewPosition(InArgs._OnUserDraggedToNewPosition)
+			[
+				InArgs._Content.Widget
+			]
+		]
+	];
+
+	SetBoxAlignmentOffset(InArgs._InitialAlignmentOffset, false);
+}
+
+FVector2f SDraggableBoxOverlay::GetBoxAlignmentOffset() const
+{
+	FVector2f AlignmentOffset = FVector2f::ZeroVector;
+
+	if (Container.IsValid())
+	{
+		switch (HorizontalAlignment)
+		{
+		case EHorizontalAlignment::HAlign_Left:
+			AlignmentOffset.X = Padding.Left;
+			break;
+		case EHorizontalAlignment::HAlign_Right:
+			AlignmentOffset.X = Padding.Right;
+			break;
+		default: break;
+		}
+
+		switch (VerticalAlignment)
+		{
+		case EVerticalAlignment::VAlign_Top:
+			AlignmentOffset.Y = Padding.Top;
+			break;
+		case EVerticalAlignment::VAlign_Bottom:
+			AlignmentOffset.Y = Padding.Bottom;
+			break;
+		default: break;
+		}
+	}
+
+	return AlignmentOffset;
+}
+
+void SDraggableBoxOverlay::SetBoxAlignmentOffset(const FVector2f& InOffset, bool bInRecomputeAnchorPoints)
+{
+	if (!Container.IsValid() || !DraggableBox.IsValid())
+	{
+		return;
+	}
+
+	FVector2f ConstrainedOffset = {
+		FMath::Max(InOffset.X, DraggableBoxDetail::DraggableBorderArea),
+		FMath::Max(InOffset.Y, DraggableBoxDetail::DraggableBorderArea)
+	};
+	if (bInRecomputeAnchorPoints)
+	{
+		ConstrainedOffset = RecomputeAnchorPoints(ConstrainedOffset);
+	}
+
+	switch (HorizontalAlignment)
+	{
+	case EHorizontalAlignment::HAlign_Left:
+		Padding.Left = ConstrainedOffset.X;
+		Padding.Right = 0.f;
+		break;
+
+	case EHorizontalAlignment::HAlign_Right:
+		Padding.Left = 0.f;
+		Padding.Right = ConstrainedOffset.X;
+		break;
+
+	default: break;
+	}
+
+	switch (VerticalAlignment)
+	{
+	case EVerticalAlignment::VAlign_Top:
+		Padding.Top = ConstrainedOffset.Y;
+		Padding.Bottom = 0.f;
+		break;
+
+	case EVerticalAlignment::VAlign_Bottom:
+		Padding.Top = 0.f;
+		Padding.Bottom = ConstrainedOffset.Y;
+		break;
+
+	default: break;
+	}
+
+	if (Container.IsValid())
+	{
+		Container->SetPadding(Padding);
+	}
+}
+
+EHorizontalAlignment SDraggableBoxOverlay::GetBoxHorizontalAlignment() const
+{
+	return HorizontalAlignment;
+}
+
+void SDraggableBoxOverlay::SetBoxHorizontalAlignment(EHorizontalAlignment InAlignment)
+{
+	HorizontalAlignment = InAlignment;
+
+	if (Container.IsValid())
+	{
+		Container->SetHAlign(InAlignment);
+	}
+}
+
+EVerticalAlignment SDraggableBoxOverlay::GetBoxVerticalAlignment() const
+{
+	return VerticalAlignment;
+}
+
+void SDraggableBoxOverlay::SetBoxVerticalAlignment(EVerticalAlignment InAlignment)
+{
+	VerticalAlignment = InAlignment;
+
+	if (Container.IsValid())
+	{
+		Container->SetVAlign(InAlignment);
+	}
+}
+
+FMargin SDraggableBoxOverlay::GetPadding() const
+{
+	return Padding;
+}
+
+FVector2f SDraggableBoxOverlay::RecomputeAnchorPoints(const FVector2f& InOffset)
+{
+	const FGeometry& MyGeometry = GetTickSpaceGeometry();
+
+	const FVector2f AvailableSpace = (MyGeometry.GetAbsoluteSize() - DraggableBox->GetTickSpaceGeometry().GetAbsoluteSize())
+		* (MyGeometry.GetLocalSize() / MyGeometry.GetAbsoluteSize());
+	const FVector2f MidPoint = AvailableSpace * 0.5f;
+	if (!ensure(!MidPoint.ContainsNaN())) // Don't call this during construction because the geometry is not initialized, yet. 
+	{
+		return InOffset;
+	}
+
+	FVector2f ConstrainedOffset = InOffset;
+	ConstrainedOffset.X = FMath::Min(ConstrainedOffset.X, AvailableSpace.X);
+	ConstrainedOffset.Y = FMath::Min(ConstrainedOffset.Y, AvailableSpace.Y);
+	if (ConstrainedOffset.X > MidPoint.X)
+	{
+		ConstrainedOffset.X = FMath::Max(AvailableSpace.X - ConstrainedOffset.X, DraggableBoxDetail::DraggableBorderArea); // Circle value around
+
+		switch (HorizontalAlignment)
+		{
+		case EHorizontalAlignment::HAlign_Left:
+			SetBoxHorizontalAlignment(EHorizontalAlignment::HAlign_Right);
+			break;
+
+		case EHorizontalAlignment::HAlign_Right:
+			SetBoxHorizontalAlignment(EHorizontalAlignment::HAlign_Left);
+			break;
+
+		default:
+			// Do nothing
+			break;
+		}
+	}
+
+	if (ConstrainedOffset.Y > MidPoint.Y)
+	{
+		ConstrainedOffset.Y = FMath::Max(AvailableSpace.Y - ConstrainedOffset.Y, DraggableBoxDetail::DraggableBorderArea); // Circle value around
+
+		switch (VerticalAlignment)
+		{
+		case EVerticalAlignment::VAlign_Top:
+			SetBoxVerticalAlignment(EVerticalAlignment::VAlign_Bottom);
+			break;
+
+		case EVerticalAlignment::VAlign_Bottom:
+			SetBoxVerticalAlignment(EVerticalAlignment::VAlign_Top);
+			break;
+
+		default:
+			// Do nothing
+			break;
+		}
+	}
+
+	return ConstrainedOffset;
+}
+
+FToolWidget_DragBoxPosition SDraggableBoxOverlay::GetDragBoxPosition() const
+{
+	return FToolWidget_DragBoxPosition{ GetBoxAlignmentOffset(), HorizontalAlignment, VerticalAlignment };
+}
+
+void SDraggableBoxOverlay::RestoreFromDragBoxPosition(const FToolWidget_DragBoxPosition& InWidgetPosition)
+{
+	SetBoxHorizontalAlignment(InWidgetPosition.HAlign);
+	SetBoxVerticalAlignment(InWidgetPosition.VAlign);
+
+	// Do not recompute the anchors. Suppose the viewport is now much smaller than it was when InWidgetPosition was saved.
+	// If the user now increases the size of the viewport, the widget should stay anchored to the same corner as when it was saved.
+	constexpr bool bRecomputeAnchors = false;
+	SetBoxAlignmentOffset(InWidgetPosition.RelativeOffset, bRecomputeAnchors);
+}
+}
\ No newline at end of file
diff --git a/Engine/Source/Developer/ToolWidgets/Public/Overlay/DragBoxPosition.h b/Engine/Source/Developer/ToolWidgets/Public/Overlay/DragBoxPosition.h
new file mode 100644
--- /dev/null
+++ b/Engine/Source/Developer/ToolWidgets/Public/Overlay/DragBoxPosition.h
@@ -0,0 +1,31 @@
+﻿// Copyright Epic Games, Inc. All Rights Reserved.
+
+#pragma once
+
+#include "Types/SlateEnums.h"
+#include "DragBoxPosition.generated.h"
+
+/** Used to restore position of widget in SDraggableBoxOverlay. */
+USTRUCT()
+struct FToolWidget_DragBoxPosition
+{
+	GENERATED_BODY()
+
+	UPROPERTY()
+	FVector2f RelativeOffset;
+
+	UPROPERTY()
+	TEnumAsByte<EHorizontalAlignment> HAlign;
+
+	UPROPERTY()
+	TEnumAsByte<EVerticalAlignment> VAlign;
+
+	FToolWidget_DragBoxPosition() = default;
+	explicit FToolWidget_DragBoxPosition(
+		const FVector2f& RelativeOffset, const TEnumAsByte<EHorizontalAlignment>& HAlign, const TEnumAsByte<EVerticalAlignment>& VAlign
+		)
+		: RelativeOffset(RelativeOffset)
+		, HAlign(HAlign)
+		, VAlign(VAlign)
+	{}
+};
diff --git a/Engine/Source/Developer/ToolWidgets/Public/Overlay/SDraggableBoxOverlay.h b/Engine/Source/Developer/ToolWidgets/Public/Overlay/SDraggableBoxOverlay.h
new file mode 100644
--- /dev/null
+++ b/Engine/Source/Developer/ToolWidgets/Public/Overlay/SDraggableBoxOverlay.h
@@ -0,0 +1,102 @@
+﻿// Copyright Epic Games, Inc. All Rights Reserved.
+
+#pragma once
+
+#include "Widgets/SCompoundWidget.h"
+#include "Layout/Margin.h"
+#include "Types/SlateEnums.h"
+#include "Widgets/SBoxPanel.h"
+#include "Widgets/DeclarativeSyntaxSupport.h"
+
+class SBox;
+struct FToolWidget_DragBoxPosition;
+
+namespace UE::ToolWidgets
+{
+/**
+ * Class that can be used to place a draggable box into a viewport or some other large widget as an overlay.
+ * Just place the widget that you want to be draggable as the contents.
+ * 
+ * This overlay will remain confined to the bounds of the viewport dimensions.
+ * It anchors the content to the corners of the parent widget: e.g. when a user drags it from the bottom-left to the bottom-right corner and
+ * drags the right side of the viewport, the widget moves relative to the right side.
+ */
+class TOOLWIDGETS_API SDraggableBoxOverlay : public SCompoundWidget
+{
+public:
+
+	SLATE_BEGIN_ARGS(SDraggableBoxOverlay)
+		: _HAlign(HAlign_Left)
+		, _VAlign(VAlign_Bottom)
+		, _InitialAlignmentOffset(FVector2f::ZeroVector)
+		, _IsDraggable(true)
+	{}
+	/** Use this to specify whether to align to the left or right. Center and Fill prevent repositioning horizontally. */
+	SLATE_ARGUMENT(EHorizontalAlignment, HAlign)
+	/** Use this to specify whether to align to the top or bottom. Center and Fill prevent repositioning vertically. */
+	SLATE_ARGUMENT(EVerticalAlignment, VAlign)
+	
+	/** The initial offset this widget should have relative to HAlign and VAlign. */
+	SLATE_ARGUMENT(FVector2f, InitialAlignmentOffset)
+	/** Whether this widget should be allowed to be dragged. */
+	SLATE_ATTRIBUTE(bool, IsDraggable)
+
+	/**
+	 * Invoked when the user has finished dragging the widget to a new position.
+	 * This is useful if you want to save the box position persistently: call GetDragBoxPosition and save it in a config file.
+	 */
+	SLATE_EVENT(FSimpleDelegate, OnUserDraggedToNewPosition)
+	
+	/** The content you want to be able to drag around in the parent widget. */
+	SLATE_DEFAULT_SLOT(FArguments, Content)
+SLATE_END_ARGS()
+
+void Construct(const FArguments& InArgs);
+
+	/** @return The current position of the box on the viewport relative to HAlign and VAlign. */
+	FVector2f GetBoxAlignmentOffset() const;
+
+	/*
+	 * Sets the content's position relative to the current anchor points.
+	 * 
+	 * @param InOffset Offset relative to the current HAlign and VAlign, e.g. X = 100 means an offset of 100 from the left or right widget side.
+	 * @param bInRecomputeAnchorPoints Whether you want HAlign and VAlign recomputed,
+	 *	e.g. if HAlign is currently left, parent is 200 wide, and InOffset.X == 150, this would recompute HAlign to be right.
+	 */
+	void SetBoxAlignmentOffset(const FVector2f& InOffset, bool bInRecomputeAnchorPoints = true);
+
+	EHorizontalAlignment GetBoxHorizontalAlignment() const;
+	void SetBoxHorizontalAlignment(EHorizontalAlignment InAlignment);
+
+	EVerticalAlignment GetBoxVerticalAlignment() const;
+	void SetBoxVerticalAlignment(EVerticalAlignment InAlignment);
+
+	/** Gets information to restore this content's position. This is useful for saving in config files. */
+	FToolWidget_DragBoxPosition GetDragBoxPosition() const;
+	/** Restores the content position. This is useful for loading a widget position saved in a config file. */
+	void RestoreFromDragBoxPosition(const FToolWidget_DragBoxPosition& InWidgetPosition);
+
+protected:
+
+	/** Contains the DraggableBox and repositions it relative to our parent. */
+	TSharedPtr<SBox> Container;
+	/** Contains the content and handles the drag-drop behaviour. */
+	TSharedPtr<SWidget> DraggableBox;
+
+	/** Affects the horizontal alignment of the content widget relative to the parent widget. */
+	EHorizontalAlignment HorizontalAlignment = HAlign_Fill;
+	/** Affects the vertical alignment of the content widget relative to the parent widget. */
+	EVerticalAlignment VerticalAlignment = VAlign_Fill;
+	/** The padding of the content widget. */
+	FMargin Padding;
+
+	FMargin GetPadding() const;
+
+	/**
+	 * Recomputes HorizontalAlignment and VerticalAlignment if InOffset is too far away from the current corner.
+	 * @param InOffset Offset relative to current HorizontalAlignment and VerticalAlignment
+	 * @return New offset relative to recomputed HorizontalAlignment and VerticalAlignment
+	 */
+	FVector2f RecomputeAnchorPoints(const FVector2f& InOffset);
+};
+}
\ No newline at end of file
