Tool Widgets: Fixed FToolWidget_DragBoxPosition members not being initialized properly.

diff --git a/Engine/Source/Developer/ToolWidgets/Public/Overlay/DragBoxPosition.h b/Engine/Source/Developer/ToolWidgets/Public/Overlay/DragBoxPosition.h
--- a/Engine/Source/Developer/ToolWidgets/Public/Overlay/DragBoxPosition.h
+++ b/Engine/Source/Developer/ToolWidgets/Public/Overlay/DragBoxPosition.h
@@ -12,13 +12,13 @@ struct FToolWidget_DragBoxPosition
 	GENERATED_BODY()
 
 	UPROPERTY()
-	FVector2f RelativeOffset;
+	FVector2f RelativeOffset = FVector2f::ZeroVector;
 
 	UPROPERTY()
-	TEnumAsByte<EHorizontalAlignment> HAlign;
+	TEnumAsByte<EHorizontalAlignment> HAlign = HAlign_Left;
 
 	UPROPERTY()
-	TEnumAsByte<EVerticalAlignment> VAlign;
+	TEnumAsByte<EVerticalAlignment> VAlign = VAlign_Bottom;
 
 	FToolWidget_DragBoxPosition() = default;
 	explicit FToolWidget_DragBoxPosition(
