PR #13130: Fix the SCheckBoxList 'toggle all' header checkbox to show an accurate value when adding or removing checkboxes at runtime.

#13130
[RN] fix ui/slate
#tests I created a slate widget for an SCheckBoxList, along with 3 buttons; one for add, one for remove, and one for remove all. I then kept adding and removing different combinations of true/false button states, removing all, etc, and made sure that the global checkbox was correct for all of them.

diff --git a/Engine/Source/Developer/ToolWidgets/Private/SCheckBoxList.cpp b/Engine/Source/Developer/ToolWidgets/Private/SCheckBoxList.cpp
--- a/Engine/Source/Developer/ToolWidgets/Private/SCheckBoxList.cpp
+++ b/Engine/Source/Developer/ToolWidgets/Private/SCheckBoxList.cpp
@@ -126,11 +126,15 @@ void SCheckBoxList::Construct(const FArguments& InArgs, const TArray<TSharedRef<
 		.HeaderRow(HeaderRowWidget)
 		.SelectionMode(ESelectionMode::None)
 	];
+
+	// Set the Global CheckBox to be correct
+	UpdateAllChecked();
 }
 
 int32 SCheckBoxList::AddItem(const FText& Text, bool bIsChecked)
 {
 	int32 ReturnValue = Items.Add(MakeShared<CheckBoxList::FItemPair>(SNew(STextBlock).Text(Text), bIsChecked));
+	UpdateAllChecked();	
 	ListView->RebuildList();
 	return ReturnValue;
 }
@@ -138,6 +142,7 @@ int32 SCheckBoxList::AddItem(const FText& Text, bool bIsChecked)
 int32 SCheckBoxList::AddItem(TSharedRef<SWidget> Widget, bool bIsChecked)
 {
 	int32 ReturnValue = Items.Add(MakeShared<CheckBoxList::FItemPair>(Widget, bIsChecked));
+	UpdateAllChecked();	
 	ListView->RebuildList();
 	return ReturnValue;
 }
@@ -145,6 +150,7 @@ int32 SCheckBoxList::AddItem(TSharedRef<SWidget> Widget, bool bIsChecked)
 void SCheckBoxList::RemoveAll()
 {
 	Items.Reset();
+	UpdateAllChecked();
 	ListView->RebuildList();
 }
 
@@ -153,6 +159,7 @@ void SCheckBoxList::RemoveItem(int32 Index)
 	if (Items.IsValidIndex(Index))
 	{
 		Items.RemoveAt(Index);
+		UpdateAllChecked();
 		ListView->RebuildList();
 	}
 }
@@ -214,4 +221,4 @@ TSharedRef<ITableRow> SCheckBoxList::HandleGenerateRow(TSharedRef<CheckBoxList::
 		.OnCheckUpdated(FSimpleDelegate::CreateLambda([this, InItem]() { OnItemCheckBox(InItem); }));
 }
 
-#undef LOCTEXT_NAMESPACE
\ No newline at end of file
+#undef LOCTEXT_NAMESPACE
diff --git a/Engine/Source/Developer/ToolWidgets/Public/SCheckBoxList.h b/Engine/Source/Developer/ToolWidgets/Public/SCheckBoxList.h
--- a/Engine/Source/Developer/ToolWidgets/Public/SCheckBoxList.h
+++ b/Engine/Source/Developer/ToolWidgets/Public/SCheckBoxList.h
@@ -65,10 +65,10 @@ private:
 
 	TSharedRef<ITableRow> HandleGenerateRow(TSharedRef<CheckBoxList::FItemPair> InItem, const TSharedRef<STableViewBase>& OwnerTable);
 
-	ECheckBoxState bAllCheckedState;
+	ECheckBoxState bAllCheckedState = ECheckBoxState::Unchecked;
 	TArray<TSharedRef<CheckBoxList::FItemPair>> Items;
 
-	const FCheckBoxStyle* CheckBoxStyle;
+	const FCheckBoxStyle* CheckBoxStyle = nullptr;
 	TSharedPtr<SListView<TSharedRef<CheckBoxList::FItemPair>>> ListView;
 
 	FOnCheckListItemStateChanged OnItemCheckStateChanged;
