Made adjustments to SCustomDialog to open it up a little better for wider usage, as well as added convenience functionality for adding hotkeys:
- Eliminated the error when not auto-generating buttons.
- Added dialog commands for Cancel and Confirm, bound by default to Escape and Enter, respectively.
- Auto-generated buttons can be assigned Cancel and Confirm roles, causing them to be triggered by the hotkeys.
- Each hotkey can have a function bound to it, for handling buttons that are just part of the Content.
- Can now set a function to get the widget to focus on activate, for when the widget to focus is dynamic and part of the Content.
- Overrode OnIsActiveChanged to updated the widget to focus on activate, as well as change activations by mouse to be considered keyboard activations to prevent the focus on activation logic from being skipped.

[RN]

diff --git a/Engine/Source/Developer/ToolWidgets/Private/Dialog/DialogCommands.cpp b/Engine/Source/Developer/ToolWidgets/Private/Dialog/DialogCommands.cpp
new file mode 100644
--- /dev/null
+++ b/Engine/Source/Developer/ToolWidgets/Private/Dialog/DialogCommands.cpp
@@ -0,0 +1,28 @@
+// Copyright Epic Games, Inc. All Rights Reserved.
+
+#include "Dialog/DialogCommands.h"
+
+#include "Styling/AppStyle.h"
+
+#define LOCTEXT_NAMESPACE "CustomDialogCommands"
+
+UE_DEFINE_TCOMMANDS(FDialogCommands)
+
+FDialogCommands::FDialogCommands()
+	: TCommands<FDialogCommands>
+	(
+		TEXT("Dialogs"), // Context name for fast lookup
+		NSLOCTEXT("Contexts", "DialogCommands", "Dialog Commands"), // Localized context name for displaying
+		TEXT("MainFrame"),
+		FAppStyle::GetAppStyleSetName() // Icon Style Set
+	)
+{
+}
+
+void FDialogCommands::RegisterCommands()
+{
+	UI_COMMAND(Cancel, "Cancel", "Cancel the current dialog", EUserInterfaceActionType::Button, FInputChord(EKeys::Escape));
+	UI_COMMAND(Confirm, "Confirm", "Confirm the current dialog", EUserInterfaceActionType::Button, FInputChord(EKeys::Enter));
+}
+
+#undef LOCTEXT_NAMESPACE
\ No newline at end of file
diff --git a/Engine/Source/Developer/ToolWidgets/Private/Dialog/SCustomDialog.cpp b/Engine/Source/Developer/ToolWidgets/Private/Dialog/SCustomDialog.cpp
--- a/Engine/Source/Developer/ToolWidgets/Private/Dialog/SCustomDialog.cpp
+++ b/Engine/Source/Developer/ToolWidgets/Private/Dialog/SCustomDialog.cpp
@@ -2,6 +2,7 @@
 
 #include "Dialog/SCustomDialog.h"
 
+#include "Dialog/DialogCommands.h"
 #include "Framework/Application/SlateApplication.h"
 #include "Framework/Docking/TabManager.h"
 #include "Styling/AppStyle.h"
@@ -26,14 +27,32 @@ SCustomDialog::FArguments& SCustomDialog::FArguments::IconBrush(FName InIconBrus
 
 void SCustomDialog::Construct(const FArguments& InArgs)
 {
-	check(InArgs._Buttons.Num() > 0);
-	
 	OnClosed = InArgs._OnClosed;
 	bAutoCloseOnButtonPress = InArgs._AutoCloseOnButtonPress;
-	
+	OnCancelHotkeyPressed = InArgs._OnCancelHotkeyPressed;
+	OnConfirmHotkeyPressed = InArgs._OnConfirmHotkeyPressed;
+	GetWidgetToFocusOnActivate = InArgs._GetWidgetToFocusOnActivate;
+
+	// Build the command list, adding the Cancel and Confirm commands from FDialogCommands
+	CommandList = MakeShareable(new FUICommandList);
+	CommandList->MapAction(FDialogCommands::Get().Cancel,
+		FUIAction(
+			FSimpleDelegate::CreateSP(this, &SCustomDialog::ExecuteCancel),
+			FCanExecuteAction::CreateSP(this, &SCustomDialog::CanExecuteCancel)
+		)
+	);
+	CommandList->MapAction(FDialogCommands::Get().Confirm,
+		FUIAction(
+			FSimpleDelegate::CreateSP(this, &SCustomDialog::ExecuteConfirm),
+			FCanExecuteAction::CreateSP(this, &SCustomDialog::CanExecuteConfirm)
+		)
+	);
+
+	TSharedPtr<SVerticalBox> VerticalBox;
 	SWindow::Construct( SWindow::FArguments(InArgs._WindowArguments)
 		.Title(InArgs._Title)
-		.SizingRule(ESizingRule::Autosized)
+		.ClientSize(InArgs._ClientSize) // If the ClientSize is left at zero, use Autosized
+		.SizingRule(InArgs._ClientSize.IsNearlyZero() ? ESizingRule::Autosized : ESizingRule::UserSized)
 		.SupportsMaximize(false)
 		.SupportsMinimize(false)
 		[
@@ -41,23 +60,27 @@ void SCustomDialog::Construct(const FArguments& InArgs)
 			.Padding(InArgs._RootPadding)
 			.BorderImage(FAppStyle::Get().GetBrush( "ToolPanel.GroupBorder" ))
 			[
-				SNew(SVerticalBox)
+				SAssignNew(VerticalBox, SVerticalBox)
 				+ SVerticalBox::Slot()
 				.FillHeight(1.0f)
 				[
 					CreateContentBox(InArgs)
 				]
-				
-				+ SVerticalBox::Slot()
-				.VAlign(InArgs._VAlignButtonBox)
-				.HAlign(InArgs._HAlignButtonBox)
-				.AutoHeight()
-				.Padding(InArgs._ButtonAreaPadding)
-				[
-					CreateButtonBox(InArgs)
-				]
 			]
 		] );
+
+	// Only create the button box if there are buttons to generate
+	if (InArgs._Buttons.Num() > 0)
+	{
+		VerticalBox->AddSlot()
+			.VAlign(InArgs._VAlignButtonBox)
+			.HAlign(InArgs._HAlignButtonBox)
+			.AutoHeight()
+			.Padding(InArgs._ButtonAreaPadding)
+			[
+				CreateButtonBox(InArgs)
+			];
+	}
 }
 
 int32 SCustomDialog::ShowModal()
@@ -66,6 +89,51 @@ int32 SCustomDialog::ShowModal()
 	return LastPressedButton;
 }
 
+void SCustomDialog::SetOnCancelHotkeyPressed(const FSimpleDelegate& InOnCancelHotkeyPressed)
+{
+	OnCancelHotkeyPressed = InOnCancelHotkeyPressed;
+}
+
+void SCustomDialog::SetOnConfirmHotkeyPressed(const FSimpleDelegate& InOnConfirmHotkeyPressed)
+{
+	OnConfirmHotkeyPressed = InOnConfirmHotkeyPressed;
+}
+
+void SCustomDialog::SetGetWidgetToFocusOnActivate(const FGetFocusWidget& InGetWidgetToFocusOnActivate)
+{
+	GetWidgetToFocusOnActivate = InGetWidgetToFocusOnActivate;
+}
+
+FReply SCustomDialog::OnKeyDown(const FGeometry& MyGeometry, const FKeyEvent& InKeyEvent)
+{
+	if (CommandList->ProcessCommandBindings(InKeyEvent))
+	{
+		return FReply::Handled();
+	}
+
+	return SWindow::OnKeyDown(MyGeometry, InKeyEvent);
+}
+
+bool SCustomDialog::OnIsActiveChanged(const FWindowActivateEvent& ActivateEvent)
+{
+	FWindowActivateEvent::EActivationType ActivationType = ActivateEvent.GetActivationType();
+
+	// If this is a mouse activation, switch to keyboard activation so focus-on-activate logic won't be skipped
+	if (ActivationType == FWindowActivateEvent::EA_ActivateByMouse)
+	{
+		ActivationType = FWindowActivateEvent::EA_Activate;
+	}
+
+	// If there's a function to get the focus on activation, use that to set the widget to focus
+	if (GetWidgetToFocusOnActivate.IsBound())
+	{
+		WidgetToFocusOnActivate = GetWidgetToFocusOnActivate.Execute();
+	}
+
+	const FWindowActivateEvent ModifiedActivateEvent(ActivationType, ActivateEvent.GetAffectedWindow());
+	return SWindow::OnIsActiveChanged(ModifiedActivateEvent);
+}
+
 void SCustomDialog::Show()
 {
 	const TSharedRef<SWindow> Window = FSlateApplication::Get().AddWindow(StaticCastSharedRef<SWindow>(this->AsShared()), true);
@@ -160,11 +228,13 @@ TSharedRef<SWidget> SCustomDialog::CreateButtonBox(const FArguments& InArgs)
 		const FButtonStyle* ButtonStyle = Button.bIsPrimary ?
 				&FAppStyle::Get().GetWidgetStyle< FButtonStyle >( "PrimaryButton" ) :
 				&FAppStyle::Get().GetWidgetStyle< FButtonStyle >("Button");
-		
+
+		FOnClicked OnClicked = FOnClicked::CreateSP(
+			this, &SCustomDialog::OnButtonClicked, Button.OnClicked, ButtonIndex);
 		TSharedRef<SButton> ButtonWidget = SNew(SButton)
 			.IsEnabled(Button.ButtonIsEnabled)
 			.ToolTipText(Button.ButtonToolTipText)
-			.OnClicked(FOnClicked::CreateSP(this, &SCustomDialog::OnButtonClicked, Button.OnClicked, ButtonIndex))
+			.OnClicked(OnClicked)
 			.ButtonStyle(ButtonStyle)
 			[
 				SNew(STextBlock)
@@ -185,12 +255,23 @@ TSharedRef<SWidget> SCustomDialog::CreateButtonBox(const FArguments& InArgs)
 		
 		const bool bIsLastButton = ButtonIndex == InArgs._Buttons.Num() - 1;
 		if (Button.bShouldFocus
-			|| (bCanFocusLastPrimary && Button.bIsPrimary)
+			|| (bCanFocusLastPrimary && (Button.bIsPrimary || Button.ButtonRole == EButtonRole::Confirm))
 			|| (bIsLastButton && !bFocusedAnyButtonYet))
 		{
 			bFocusedAnyButtonYet = true;
 			SetWidgetToFocusOnActivate(ButtonWidget);
 		}
+
+		// Remember the buttons for later, so we can trigger them via hotkeys
+		// If multiple buttons have identical roles, only the last will matter
+		if (Button.ButtonRole == EButtonRole::Cancel)
+		{
+			CancelButtonOnClicked = OnClicked;
+		}
+		else if (Button.ButtonRole == EButtonRole::Confirm)
+		{
+			ConfirmButtonOnClicked = OnClicked;
+		}
 	}
 	return ButtonBox;
 }
@@ -208,3 +289,39 @@ FReply SCustomDialog::OnButtonClicked(FSimpleDelegate OnClicked, int32 ButtonInd
 	OnClicked.ExecuteIfBound();
 	return FReply::Handled();
 }
+
+bool SCustomDialog::CanExecuteCancel() const
+{
+	// Return true if Cancel button functionality is available
+	return CancelButtonOnClicked.IsBound() || OnCancelHotkeyPressed.IsBound();
+}
+
+void SCustomDialog::ExecuteCancel()
+{
+	// If a generated button was given the Cancel role, use that button's OnClick
+	if (CancelButtonOnClicked.IsBound() && CancelButtonOnClicked.Execute().IsEventHandled())
+	{
+		return;
+	}
+
+	// Otherwise, try using the function given for Cancel
+	OnCancelHotkeyPressed.ExecuteIfBound();
+}
+
+bool SCustomDialog::CanExecuteConfirm() const
+{
+	// Return true if a Confirm button functionality is available
+	return ConfirmButtonOnClicked.IsBound() || OnConfirmHotkeyPressed.IsBound();
+}
+
+void SCustomDialog::ExecuteConfirm()
+{
+	// If a generated button was given the Confirm role, use that button's OnClick
+	if (ConfirmButtonOnClicked.IsBound() && ConfirmButtonOnClicked.Execute().IsEventHandled())
+	{
+		return;
+	}
+
+	// Otherwise, try using the function given for Confirm
+	OnConfirmHotkeyPressed.ExecuteIfBound();
+}
diff --git a/Engine/Source/Developer/ToolWidgets/Private/ToolWidgetsModule.cpp b/Engine/Source/Developer/ToolWidgets/Private/ToolWidgetsModule.cpp
--- a/Engine/Source/Developer/ToolWidgets/Private/ToolWidgetsModule.cpp
+++ b/Engine/Source/Developer/ToolWidgets/Private/ToolWidgetsModule.cpp
@@ -1,5 +1,6 @@
 // Copyright Epic Games, Inc. All Rights Reserved.
 
+#include "Dialog/DialogCommands.h"
 #include "Modules/ModuleInterface.h"
 #include "Modules/ModuleManager.h"
 
@@ -10,7 +11,11 @@ public:
 
 	// IModuleInterface interface
 	
-	virtual void StartupModule() override { }
+	virtual void StartupModule() override
+	{
+		// Register the dialog commands
+		FDialogCommands::Register();
+	}
 	virtual void ShutdownModule() override { }
 };
 
diff --git a/Engine/Source/Developer/ToolWidgets/Public/Dialog/DialogCommands.h b/Engine/Source/Developer/ToolWidgets/Public/Dialog/DialogCommands.h
new file mode 100644
--- /dev/null
+++ b/Engine/Source/Developer/ToolWidgets/Public/Dialog/DialogCommands.h
@@ -0,0 +1,33 @@
+// Copyright Epic Games, Inc. All Rights Reserved.
+
+#pragma once
+
+#include "Framework/Commands/Commands.h"
+#include "Framework/Commands/UICommandInfo.h"
+#include "Templates/SharedPointer.h"
+
+#define UE_API TOOLWIDGETS_API
+
+class FUICommandInfo;
+
+UE_DECLARE_TCOMMANDS(class FDialogCommands, UE_API)
+
+/**
+ * Class containing commands for dialog actions common to most dialogs
+ */
+	class FDialogCommands : public TCommands<FDialogCommands>
+{
+public:
+	UE_API FDialogCommands();
+
+	/** Trigger the default cancellation action for the current dialog */
+	TSharedPtr<FUICommandInfo> Cancel;
+
+	/** Trigger the default confirmation action for the current dialog */
+	TSharedPtr<FUICommandInfo> Confirm;
+
+	/** Registers our commands with the binding system */
+	UE_API virtual void RegisterCommands() override;
+};
+
+#undef UE_API
\ No newline at end of file
diff --git a/Engine/Source/Developer/ToolWidgets/Public/Dialog/SCustomDialog.h b/Engine/Source/Developer/ToolWidgets/Public/Dialog/SCustomDialog.h
--- a/Engine/Source/Developer/ToolWidgets/Public/Dialog/SCustomDialog.h
+++ b/Engine/Source/Developer/ToolWidgets/Public/Dialog/SCustomDialog.h
@@ -3,14 +3,19 @@
 #pragma once
 
 #include "CoreMinimal.h"
+#include "Framework/SlateDelegates.h"
+#include "Framework/Commands/UICommandList.h"
 #include "Styling/AppStyle.h"
 #include "Widgets/SWindow.h"
 
 #define UE_API TOOLWIDGETS_API
 
+/** Delegate for functions to get the appropriate widget to focus on activation */
+DECLARE_DELEGATE_RetVal(TSharedPtr<SWidget>, FGetFocusWidget);
+
 /**
  * This is a custom dialog class, which allows any Slate widget to be used as the contents,
- * with any number of buttons that have any text. 
+ * with any number of auto-generated buttons that have any text. 
  * It also supports adding a custom icon to the dialog.
  * 
  * Usage:
@@ -33,13 +38,23 @@
 class SCustomDialog : public SWindow
 {
 public:
-	
+
+	/** Determines whether a button is triggered via the Cancel or Confirm hotkeys.
+	 * If more than one button is assigned the same role, the last with that role will be used. */
+	enum class EButtonRole
+	{
+		Custom,
+		Cancel,
+		Confirm
+	};
+
 	struct FButton
 	{
-		FButton(TAttribute<FText> InButtonText, FSimpleDelegate InOnClicked = FSimpleDelegate())
+		FButton(TAttribute<FText> InButtonText, FSimpleDelegate InOnClicked = FSimpleDelegate(), const EButtonRole InButtonRole = EButtonRole::Custom)
 			: ButtonText(MoveTemp(InButtonText))
 			, ButtonIsEnabled(true)
 			, OnClicked(MoveTemp(InOnClicked))
+			, ButtonRole(InButtonRole)
 		{}
 
 		/** Set whether this button enabled */
@@ -54,10 +69,15 @@ public:
 		/** Called when the button is clicked */
 		FButton& SetOnClicked(FSimpleDelegate InOnClicked) { OnClicked = MoveTemp(InOnClicked); return *this; }
 
+		/** Button role is used to determine whether the button should be triggered via the Cancel and Confirm hotkeys.
+		 * If more than one button is assigned the same role, the last with that role will be used. */
+		FButton& SetButtonRole(const EButtonRole InButtonRole) { ButtonRole = InButtonRole; return *this; }
+
 		/**
 		 * Whether to focus this button. Focus rules:
 		 *  1: If a button has SetFocus(), use the last one
 		 *  2: If a button is marked primary, use the last one
+		 *  3: If a button has the Confirm role, use the last one
 		 *  3: Otherwise, use the last button
 		 */
 		FButton& SetFocus() { bShouldFocus = true; return *this; }
@@ -66,6 +86,7 @@ public:
 		TAttribute<FText> ButtonToolTipText;
 		TAttribute<bool> ButtonIsEnabled;
 		FSimpleDelegate OnClicked;
+		EButtonRole ButtonRole;
 		bool bIsPrimary = false;
 		bool bShouldFocus = false;
 	};
@@ -81,6 +102,7 @@ public:
 		, _ScrollBoxMaxHeight(300.f)
 		, _HAlignContent(HAlign_Left)
 		, _VAlignContent(VAlign_Center)
+		, _ClientSize(FVector2f::ZeroVector)
 	{
 		_AccessibleParams = FAccessibleWidgetData(EAccessibleBehavior::Auto);
 	}
@@ -102,7 +124,7 @@ public:
 		/** Provides default values for SWindow::FArguments not overriden by SCustomDialog. */
 		SLATE_ARGUMENT(SWindow::FArguments, WindowArguments)
 
-		/** Whether to automatically close this window when any button is pressed (default: true) */
+		/** Whether to automatically close this window when any auto-generated button is pressed (default: true) */
 		SLATE_ARGUMENT(bool, AutoCloseOnButtonPress)
 
 		/*********** Cosmetic ***********/
@@ -120,7 +142,7 @@ public:
 		SLATE_ARGUMENT(EVerticalAlignment, VAlignIcon)
 	
 	
-		/** Custom widget placed before the buttons */
+		/** Custom widget placed before the auto-generated buttons */
 		SLATE_NAMED_SLOT(FArguments, BeforeButtons)
 
 		/** HAlign to use for Button Box slot (default: HAlign_Left) */
@@ -152,6 +174,17 @@ public:
 		/** VAlign to use for Content slot (default: VAlign_Center) */
 		SLATE_ARGUMENT(EVerticalAlignment, VAlignContent)
 
+		/** What the initial size of the window should be. If left at Zero, the window will be AutoSized. */
+		SLATE_ARGUMENT(FVector2D, ClientSize)
+
+		/** Function to call when the Cancel hotkey (Escape) is pressed */
+		SLATE_ARGUMENT(FSimpleDelegate, OnCancelHotkeyPressed)
+
+		/** Function to call when the Confirm hotkey (Enter) is pressed */
+		SLATE_ARGUMENT(FSimpleDelegate, OnConfirmHotkeyPressed)
+
+		/** Function to call when activating to get the widget to focus, for use when the widget needed is dynamic */
+		SLATE_ARGUMENT(FGetFocusWidget, GetWidgetToFocusOnActivate)
 
 		/********** Legacy - do not use **********/
 		
@@ -176,19 +209,56 @@ public:
 	 */
 	UE_API int32 ShowModal();
 
+	/** Enables the Cancel hotkey (Escape) and assigns the function to be executed */
+	UE_API void SetOnCancelHotkeyPressed(const FSimpleDelegate& InOnCancelHotkeyPressed);
+
+	/** Enables the Confirm hotkey (Enter) and assigns the function to be executed */
+	UE_API void SetOnConfirmHotkeyPressed(const FSimpleDelegate& InOnConfirmHotkeyPressed);
+
+	/** Sets the function to call when activating to get the widget to focus */
+	UE_API void SetGetWidgetToFocusOnActivate(const FGetFocusWidget& InGetWidgetToFocusOnActivate);
+
+	UE_API virtual FReply OnKeyDown(const FGeometry& MyGeometry, const FKeyEvent& InKeyEvent) override;
+
+	UE_API virtual bool OnIsActiveChanged(const FWindowActivateEvent& ActivateEvent) override;
+
 private:
 
 	TSharedRef<SWidget> CreateContentBox(const FArguments& InArgs);
 	TSharedRef<SWidget> CreateButtonBox(const FArguments& InArgs);
 	
 	FReply OnButtonClicked(FSimpleDelegate OnClicked, int32 ButtonIndex);
+
+	bool CanExecuteCancel() const;
+	void ExecuteCancel();
+
+	bool CanExecuteConfirm() const;
+	void ExecuteConfirm();
 	
 	/** The index of the button that was pressed last. */
 	int32 LastPressedButton = -1;
 
 	FSimpleDelegate OnClosed;
 
-	bool bAutoCloseOnButtonPress;
+	bool bAutoCloseOnButtonPress = true;
+
+	/** Generated Cancel button's OnClick */
+	FOnClicked CancelButtonOnClicked;
+
+	/** Generated Confirm button's OnClick */
+	FOnClicked ConfirmButtonOnClicked;
+
+	// Function to call when the Cancel hotkey (Escape) is pressed 
+	FSimpleDelegate OnCancelHotkeyPressed;
+
+	// Function to call when the Confirm hotkey (Enter) is pressed 
+	FSimpleDelegate OnConfirmHotkeyPressed;
+
+	/** Function to call when activating to get the widget to focus */
+	FGetFocusWidget GetWidgetToFocusOnActivate;
+
+	/** Command list for the dialog */
+	TSharedPtr<FUICommandList> CommandList;
 };
 
 #undef UE_API
diff --git a/Engine/Source/Editor/UnrealEd/Private/EditorEngine.cpp b/Engine/Source/Editor/UnrealEd/Private/EditorEngine.cpp
--- a/Engine/Source/Editor/UnrealEd/Private/EditorEngine.cpp
+++ b/Engine/Source/Editor/UnrealEd/Private/EditorEngine.cpp
@@ -1520,7 +1520,8 @@ void UEditorEngine::LoadDefaultEditorModules()
 			TEXT("StructUtilsTestSuite"),
 			TEXT("MassEntityEditor"),
 			TEXT("SVGDistanceField"),
-			TEXT("DataHierarchyEditor")
+			TEXT("DataHierarchyEditor"),
+			TEXT("ToolWidgets")
 		};
 
 	FScopedSlowTask ModuleSlowTask((float)UE_ARRAY_COUNT(ModuleNames));
