Fixed issue where SearchableComboBox couldn't be set back to the same value after resetting to default

diff --git a/Engine/Source/Developer/ToolWidgets/Private/SSearchableComboBox.cpp b/Engine/Source/Developer/ToolWidgets/Private/SSearchableComboBox.cpp
--- a/Engine/Source/Developer/ToolWidgets/Private/SSearchableComboBox.cpp
+++ b/Engine/Source/Developer/ToolWidgets/Private/SSearchableComboBox.cpp
@@ -109,6 +109,7 @@ void SSearchableComboBox::Construct(const FArguments& InArgs)
 void SSearchableComboBox::ClearSelection()
 {
 	ComboListView->ClearSelection();
+	SelectedItem = nullptr;
 }
 
 void SSearchableComboBox::SetSelectedItem(TSharedPtr<FString> InSelectedItem, ESelectInfo::Type InSelectInfo)
