Cleaned up some of the modifications made to SCustomDialog to better handle extensibility and clarify usage.

[RN]

diff --git a/Engine/Source/Developer/ToolWidgets/Private/Dialog/SCustomDialog.cpp b/Engine/Source/Developer/ToolWidgets/Private/Dialog/SCustomDialog.cpp
--- a/Engine/Source/Developer/ToolWidgets/Private/Dialog/SCustomDialog.cpp
+++ b/Engine/Source/Developer/ToolWidgets/Private/Dialog/SCustomDialog.cpp
@@ -91,11 +91,17 @@ int32 SCustomDialog::ShowModal()
 
 void SCustomDialog::SetOnCancelHotkeyPressed(const FSimpleDelegate& InOnCancelHotkeyPressed)
 {
+	// Don't use in conjunction with an auto-generated Cancel button
+	ensureMsgf(!bGeneratedCancelButton, TEXT("Cannot use both a Cancel function and auto-generated button with the Cancel role"));
+
 	OnCancelHotkeyPressed = InOnCancelHotkeyPressed;
 }
 
 void SCustomDialog::SetOnConfirmHotkeyPressed(const FSimpleDelegate& InOnConfirmHotkeyPressed)
 {
+	// Don't use in conjunction with an auto-generated Confirm button
+	ensureMsgf(!bGeneratedConfirmButton, TEXT("Cannot use both a Confirm function and auto-generated button with the Confirm role"));
+
 	OnConfirmHotkeyPressed = InOnConfirmHotkeyPressed;
 }
 
@@ -116,21 +122,20 @@ FReply SCustomDialog::OnKeyDown(const FGeometry& MyGeometry, const FKeyEvent& In
 
 bool SCustomDialog::OnIsActiveChanged(const FWindowActivateEvent& ActivateEvent)
 {
-	FWindowActivateEvent::EActivationType ActivationType = ActivateEvent.GetActivationType();
-
-	// If this is a mouse activation, switch to keyboard activation so focus-on-activate logic won't be skipped
-	if (ActivationType == FWindowActivateEvent::EA_ActivateByMouse)
-	{
-		ActivationType = FWindowActivateEvent::EA_Activate;
-	}
-
 	// If there's a function to get the focus on activation, use that to set the widget to focus
 	if (GetWidgetToFocusOnActivate.IsBound())
 	{
 		WidgetToFocusOnActivate = GetWidgetToFocusOnActivate.Execute();
 	}
 
-	const FWindowActivateEvent ModifiedActivateEvent(ActivationType, ActivateEvent.GetAffectedWindow());
+	FWindowActivateEvent ModifiedActivateEvent{ ActivateEvent };
+
+	// If this is a mouse activation, switch to keyboard activation so focus-on-activate logic won't be skipped
+	if (ActivateEvent.GetActivationType() == FWindowActivateEvent::EA_ActivateByMouse)
+	{
+		ModifiedActivateEvent.SetActivationType(FWindowActivateEvent::EA_Activate);
+	}
+
 	return SWindow::OnIsActiveChanged(ModifiedActivateEvent);
 }
 
@@ -229,12 +234,11 @@ TSharedRef<SWidget> SCustomDialog::CreateButtonBox(const FArguments& InArgs)
 				&FAppStyle::Get().GetWidgetStyle< FButtonStyle >( "PrimaryButton" ) :
 				&FAppStyle::Get().GetWidgetStyle< FButtonStyle >("Button");
 
-		FOnClicked OnClicked = FOnClicked::CreateSP(
-			this, &SCustomDialog::OnButtonClicked, Button.OnClicked, ButtonIndex);
 		TSharedRef<SButton> ButtonWidget = SNew(SButton)
 			.IsEnabled(Button.ButtonIsEnabled)
 			.ToolTipText(Button.ButtonToolTipText)
-			.OnClicked(OnClicked)
+			.OnClicked(FOnClicked::CreateSP(
+				this, &SCustomDialog::OnButtonClicked, Button.OnClicked, ButtonIndex))
 			.ButtonStyle(ButtonStyle)
 			[
 				SNew(STextBlock)
@@ -263,14 +267,23 @@ TSharedRef<SWidget> SCustomDialog::CreateButtonBox(const FArguments& InArgs)
 		}
 
 		// Remember the buttons for later, so we can trigger them via hotkeys
-		// If multiple buttons have identical roles, only the last will matter
 		if (Button.ButtonRole == EButtonRole::Cancel)
 		{
-			CancelButtonOnClicked = OnClicked;
+			// There should only be one auto-generated button with the Cancel role or a Cancel function
+			ensureMsgf(!bGeneratedCancelButton && !OnCancelHotkeyPressed.IsBound(), TEXT("Only use one Cancel function or Cancel role button."));
+
+			bGeneratedCancelButton = true;
+			OnCancelHotkeyPressed = FSimpleDelegate::CreateSP(
+				this, &SCustomDialog::OnHotkeyPressed, Button.OnClicked, ButtonIndex);
 		}
 		else if (Button.ButtonRole == EButtonRole::Confirm)
 		{
-			ConfirmButtonOnClicked = OnClicked;
+			// There should only be one auto-generated button with the Confirm role or a Confirm function
+			ensureMsgf(!bGeneratedConfirmButton && !OnConfirmHotkeyPressed.IsBound(), TEXT("Only use one Confirm function or Confirm role button."));
+
+			bGeneratedConfirmButton = true;
+			OnConfirmHotkeyPressed = FSimpleDelegate::CreateSP(
+				this, &SCustomDialog::OnHotkeyPressed, Button.OnClicked, ButtonIndex);
 		}
 	}
 	return ButtonBox;
@@ -278,6 +291,13 @@ TSharedRef<SWidget> SCustomDialog::CreateButtonBox(const FArguments& InArgs)
 
 /** Handle the button being clicked */
 FReply SCustomDialog::OnButtonClicked(FSimpleDelegate OnClicked, int32 ButtonIndex)
+{
+	OnHotkeyPressed(OnClicked, ButtonIndex);
+	return FReply::Handled();
+}
+
+/** Handle the hotkey being pressed */
+void SCustomDialog::OnHotkeyPressed(FSimpleDelegate OnClicked, int32 ButtonIndex)
 {
 	LastPressedButton = ButtonIndex;
 
@@ -287,41 +307,28 @@ FReply SCustomDialog::OnButtonClicked(FSimpleDelegate OnClicked, int32 ButtonInd
 	}
 
 	OnClicked.ExecuteIfBound();
-	return FReply::Handled();
 }
 
 bool SCustomDialog::CanExecuteCancel() const
 {
-	// Return true if Cancel button functionality is available
-	return CancelButtonOnClicked.IsBound() || OnCancelHotkeyPressed.IsBound();
+	// Return true if Cancel button function is available
+	return OnCancelHotkeyPressed.IsBound();
 }
 
 void SCustomDialog::ExecuteCancel()
 {
-	// If a generated button was given the Cancel role, use that button's OnClick
-	if (CancelButtonOnClicked.IsBound() && CancelButtonOnClicked.Execute().IsEventHandled())
-	{
-		return;
-	}
-
-	// Otherwise, try using the function given for Cancel
+	// Try using the Cancel function
 	OnCancelHotkeyPressed.ExecuteIfBound();
 }
 
 bool SCustomDialog::CanExecuteConfirm() const
 {
-	// Return true if a Confirm button functionality is available
-	return ConfirmButtonOnClicked.IsBound() || OnConfirmHotkeyPressed.IsBound();
+	// Return true if a Confirm button function is available
+	return OnConfirmHotkeyPressed.IsBound();
 }
 
 void SCustomDialog::ExecuteConfirm()
 {
-	// If a generated button was given the Confirm role, use that button's OnClick
-	if (ConfirmButtonOnClicked.IsBound() && ConfirmButtonOnClicked.Execute().IsEventHandled())
-	{
-		return;
-	}
-
-	// Otherwise, try using the function given for Confirm
+	// Try using the Confirm function
 	OnConfirmHotkeyPressed.ExecuteIfBound();
 }
diff --git a/Engine/Source/Developer/ToolWidgets/Public/Dialog/SCustomDialog.h b/Engine/Source/Developer/ToolWidgets/Public/Dialog/SCustomDialog.h
--- a/Engine/Source/Developer/ToolWidgets/Public/Dialog/SCustomDialog.h
+++ b/Engine/Source/Developer/ToolWidgets/Public/Dialog/SCustomDialog.h
@@ -70,7 +70,7 @@ public:
 		FButton& SetOnClicked(FSimpleDelegate InOnClicked) { OnClicked = MoveTemp(InOnClicked); return *this; }
 
 		/** Button role is used to determine whether the button should be triggered via the Cancel and Confirm hotkeys.
-		 * If more than one button is assigned the same role, the last with that role will be used. */
+		 * The Cancel and Confirm roles should each only be used on a single button. */
 		FButton& SetButtonRole(const EButtonRole InButtonRole) { ButtonRole = InButtonRole; return *this; }
 
 		/**
@@ -209,10 +209,10 @@ public:
 	 */
 	UE_API int32 ShowModal();
 
-	/** Enables the Cancel hotkey (Escape) and assigns the function to be executed */
+	/** Assigns a function to the Cancel hotkey */
 	UE_API void SetOnCancelHotkeyPressed(const FSimpleDelegate& InOnCancelHotkeyPressed);
 
-	/** Enables the Confirm hotkey (Enter) and assigns the function to be executed */
+	/** Assigns a function to the Confirm hotkey */
 	UE_API void SetOnConfirmHotkeyPressed(const FSimpleDelegate& InOnConfirmHotkeyPressed);
 
 	/** Sets the function to call when activating to get the widget to focus */
@@ -228,6 +228,7 @@ private:
 	TSharedRef<SWidget> CreateButtonBox(const FArguments& InArgs);
 	
 	FReply OnButtonClicked(FSimpleDelegate OnClicked, int32 ButtonIndex);
+	void OnHotkeyPressed(FSimpleDelegate OnClicked, int32 ButtonIndex);
 
 	bool CanExecuteCancel() const;
 	void ExecuteCancel();
@@ -240,13 +241,14 @@ private:
 
 	FSimpleDelegate OnClosed;
 
+	// Whether to close the window when any button is pressed or triggered via hotkey
 	bool bAutoCloseOnButtonPress = true;
 
-	/** Generated Cancel button's OnClick */
-	FOnClicked CancelButtonOnClicked;
+	// Whether a Cancel button has been generated
+	bool bGeneratedCancelButton = false;
 
-	/** Generated Confirm button's OnClick */
-	FOnClicked ConfirmButtonOnClicked;
+	// Whether a Confirm button has been generated
+	bool bGeneratedConfirmButton = false;
 
 	// Function to call when the Cancel hotkey (Escape) is pressed 
 	FSimpleDelegate OnCancelHotkeyPressed;
diff --git a/Engine/Source/Runtime/SlateCore/Public/Input/Events.h b/Engine/Source/Runtime/SlateCore/Public/Input/Events.h
--- a/Engine/Source/Runtime/SlateCore/Public/Input/Events.h
+++ b/Engine/Source/Runtime/SlateCore/Public/Input/Events.h
@@ -1226,6 +1226,12 @@ public:
 		return ActivationType;
 	}
 
+	/** Sets what actually happened to the window (e.g. Activated, Deactivated, Activated by a mouse click) */
+	void SetActivationType(EActivationType InActivationType)
+	{
+		ActivationType = InActivationType;
+	}
+
 	/** The window that this activation/deactivation happened to */
 	TSharedRef<SWindow> GetAffectedWindow() const
 	{
