Fix SSimpleTimeSlider not firing its slider begin movement delegate

diff --git a/Engine/Source/Developer/ToolWidgets/Private/SSimpleTimeSlider.cpp b/Engine/Source/Developer/ToolWidgets/Private/SSimpleTimeSlider.cpp
--- a/Engine/Source/Developer/ToolWidgets/Private/SSimpleTimeSlider.cpp
+++ b/Engine/Source/Developer/ToolWidgets/Private/SSimpleTimeSlider.cpp
@@ -19,6 +19,8 @@ void SSimpleTimeSlider::Construct( const SSimpleTimeSlider::FArguments& InArgs )
 	CursorSize = InArgs._CursorSize;
 	MirrorLabels = InArgs._MirrorLabels;
     OnScrubPositionChanged = InArgs._OnScrubPositionChanged;
+	OnBeginScrubberMovement = InArgs._OnBeginScrubberMovement;
+	OnEndScrubberMovement = InArgs._OnEndScrubberMovement;
     OnViewRangeChanged = InArgs._OnViewRangeChanged;
 	ClampRangeHighlightColor = InArgs._ClampRangeHighlightColor;
 	ClampRangeHighlightSize = InArgs._ClampRangeHighlightSize;
@@ -354,6 +356,7 @@ FReply SSimpleTimeSlider::OnMouseButtonUp(const FGeometry& MyGeometry, const FPo
 		}
 
 		bDraggingScrubber = false;
+		OnEndScrubberMovement.ExecuteIfBound();
 		return FReply::Handled().ReleaseMouseCapture();
 
 	}
@@ -432,6 +435,7 @@ FReply SSimpleTimeSlider::OnMouseMove( const FGeometry& MyGeometry, const FPoint
 				if ( DistanceDragged > 0/*FSlateApplication::Get().GetDragTriggerDistance()*/ )
 				{
 					bDraggingScrubber = true;
+					OnBeginScrubberMovement.ExecuteIfBound();
 				}
 			}
 			else
diff --git a/Engine/Source/Developer/ToolWidgets/Public/SSimpleTimeSlider.h b/Engine/Source/Developer/ToolWidgets/Public/SSimpleTimeSlider.h
--- a/Engine/Source/Developer/ToolWidgets/Public/SSimpleTimeSlider.h
+++ b/Engine/Source/Developer/ToolWidgets/Public/SSimpleTimeSlider.h
@@ -187,6 +187,10 @@ protected:
 	TAttribute<bool> MirrorLabels;
 	
 	FOnScrubPositionChanged OnScrubPositionChanged;
+	/** Called right before the scrubber begins to move */
+	FSimpleDelegate OnBeginScrubberMovement;
+	/** Delegate called right after the scrubber handle is released by the user */
+	FSimpleDelegate OnEndScrubberMovement;
 	FOnRangeChanged OnViewRangeChanged;
 
 	/** Brush for drawing an upwards facing scrub handle */
