[Sidebar] Support "AutoSize" and "Fill" drawer section content

diff --git a/Engine/Plugins/VirtualProduction/Avalanche/Source/AvalancheSequencer/Private/DetailsView/Section/AvaSequenceTreeDetails.h b/Engine/Plugins/VirtualProduction/Avalanche/Source/AvalancheSequencer/Private/DetailsView/Section/AvaSequenceTreeDetails.h
--- a/Engine/Plugins/VirtualProduction/Avalanche/Source/AvalancheSequencer/Private/DetailsView/Section/AvaSequenceTreeDetails.h
+++ b/Engine/Plugins/VirtualProduction/Avalanche/Source/AvalancheSequencer/Private/DetailsView/Section/AvaSequenceTreeDetails.h
@@ -22,6 +22,7 @@ public:
 	virtual FName GetSectionId() const override;
 	virtual FText GetSectionDisplayText() const override;
 	virtual bool ShouldShowSection() const override;
+	virtual TOptional<float> GetSectionFillHeight() const override { return 1.f; }
 	virtual int32 GetSortOrder() const override;
 	virtual TSharedRef<SWidget> CreateContentWidget() override;
 	//~ End ISidebarDrawerContent
diff --git a/Engine/Source/Developer/ToolWidgets/Private/Sidebar/SSidebarDrawerContent.cpp b/Engine/Source/Developer/ToolWidgets/Private/Sidebar/SSidebarDrawerContent.cpp
--- a/Engine/Source/Developer/ToolWidgets/Private/Sidebar/SSidebarDrawerContent.cpp
+++ b/Engine/Source/Developer/ToolWidgets/Private/Sidebar/SSidebarDrawerContent.cpp
@@ -41,6 +41,7 @@ void SSidebarDrawerContent::Construct(const FArguments& InArgs, const TWeakPtr<F
 		]
 		+ SVerticalBox::Slot()
 		.FillHeight(1.f)
+		.VAlign(VAlign_Fill)
 		[
 			SAssignNew(ContentBox, SVerticalBox)
 		]
@@ -252,7 +253,6 @@ TArray<TSharedRef<ISidebarDrawerContent>> SSidebarDrawerContent::GetOrderedSecti
 void SSidebarDrawerContent::AddContentSlot(const TSharedRef<ISidebarDrawerContent>& InDrawerContent)
 {
 	SVerticalBox::FSlot::FSlotArguments& Slot = ContentBox->AddSlot()
-		.FillHeight(1.f) // Remove this after jira UE-237053 has been resolved to allow drawers to specify FillHeight through the ISidebarDrawerContent interface
 		.Padding(0.f, 0.f, 0.f, 2.f)
 		[
 			SNew(SBox)
@@ -263,23 +263,15 @@ void SSidebarDrawerContent::AddContentSlot(const TSharedRef<ISidebarDrawerConten
 			]
 		];
 
-	/*
-	 * Not using this until jira UE-237053 is resolved.
-	 * Currently, this is not being used anywhere, but ISidebarDrawerContent::GetSectionFillSize
-	 * should be added to allow a drawer to specify the fill size themselves.
-	 * Motion Design, for example, could use this to better display its sections in its sidebar.
-	 * There is still some discussion on if the Motion Design sidebar is going to remain being used
-	 * or if it will be integrated into the Navigation Tool to display embedded sequences.
-	 */
-	/*const TOptional<float> FillSize = InDrawerContent->GetSectionFillSize();
-	if (FillSize.IsSet())
+	const TOptional<float> FillHeight = InDrawerContent->GetSectionFillHeight();
+	if (FillHeight.IsSet())
 	{
-		Slot.FillHeight(FillSize.Get(1.f));
+		Slot.FillHeight(FillHeight.GetValue());
 	}
 	else
 	{
 		Slot.AutoHeight();
-	}*/
+	}
 }
 
 #undef LOCTEXT_NAMESPACE 
diff --git a/Engine/Source/Developer/ToolWidgets/Private/Sidebar/SSidebarDrawerContent.h b/Engine/Source/Developer/ToolWidgets/Private/Sidebar/SSidebarDrawerContent.h
--- a/Engine/Source/Developer/ToolWidgets/Private/Sidebar/SSidebarDrawerContent.h
+++ b/Engine/Source/Developer/ToolWidgets/Private/Sidebar/SSidebarDrawerContent.h
@@ -5,9 +5,9 @@
 #include "Containers/Array.h"
 #include "Styling/SlateTypes.h"
 #include "Templates/SharedPointer.h"
-#include "UObject/NameTypes.h"
 #include "Widgets/SCompoundWidget.h"
 
+class FName;
 class FSidebarDrawer;
 class ISidebarDrawerContent;
 class SVerticalBox;
diff --git a/Engine/Source/Developer/ToolWidgets/Public/Sidebar/ISidebarDrawerContent.h b/Engine/Source/Developer/ToolWidgets/Public/Sidebar/ISidebarDrawerContent.h
--- a/Engine/Source/Developer/ToolWidgets/Public/Sidebar/ISidebarDrawerContent.h
+++ b/Engine/Source/Developer/ToolWidgets/Public/Sidebar/ISidebarDrawerContent.h
@@ -2,6 +2,7 @@
 
 #pragma once
 
+#include "Misc/Optional.h"
 #include "Templates/SharedPointer.h"
 
 class FName;
@@ -25,6 +26,9 @@ public:
 	/** Returns true if this section should be displayed. */
 	virtual bool ShouldShowSection() const { return true; }
 
+	/** Sets the optional slot height stretch coefficient. If not set, auto sizes the slot. */
+	virtual TOptional<float> GetSectionFillHeight() const { return TOptional<float>(); }
+
 	/** Gets the section sort order. Lower values will be put at the above higher values. */
 	virtual int32 GetSortOrder() const { return 0; }
 
