Fixed unreliable custom dialog button placement.

diff --git a/Engine/Source/Developer/ToolWidgets/Public/Dialog/SCustomDialog.h b/Engine/Source/Developer/ToolWidgets/Public/Dialog/SCustomDialog.h
--- a/Engine/Source/Developer/ToolWidgets/Public/Dialog/SCustomDialog.h
+++ b/Engine/Source/Developer/ToolWidgets/Public/Dialog/SCustomDialog.h
@@ -96,6 +96,8 @@ public:
 		, _Icon(nullptr)
 		, _HAlignIcon(HAlign_Left)
 		, _VAlignIcon(VAlign_Center)
+		, _HAlignButtonBox(HAlign_Right)
+		, _VAlignButtonBox(VAlign_Center)
 		, _RootPadding(FMargin(4.f))
 		, _ButtonAreaPadding(FMargin(20.f, 16.f, 0.f, 0.f))
 		, _UseScrollBox(true)
@@ -145,7 +147,7 @@ public:
 		/** Custom widget placed before the auto-generated buttons */
 		SLATE_NAMED_SLOT(FArguments, BeforeButtons)
 
-		/** HAlign to use for Button Box slot (default: HAlign_Left) */
+		/** HAlign to use for Button Box slot (default: HAlign_Right) */
 		SLATE_ARGUMENT(EHorizontalAlignment, HAlignButtonBox)
 
 		/** VAlign to use for Button Box  slot (default: VAlign_Center) */
