[FilterBar Improvements] Various QoL and UX improvements
- Added Tooltips to the the create/edit custom text filter menu
- Removed the option to select a color for custom filters if the pill style is basic (doesn't display colors)
- Custom Filter submenu entries now highlight on hover
- Fixed Alt-Click leaves filter in hovered state

diff --git a/Engine/Source/Developer/ToolWidgets/Private/Filters/SCustomTextFilterDialog.cpp b/Engine/Source/Developer/ToolWidgets/Private/Filters/SCustomTextFilterDialog.cpp
--- a/Engine/Source/Developer/ToolWidgets/Private/Filters/SCustomTextFilterDialog.cpp
+++ b/Engine/Source/Developer/ToolWidgets/Private/Filters/SCustomTextFilterDialog.cpp
@@ -15,6 +15,7 @@ void SCustomTextFilterDialog::Construct( const FArguments& InArgs )
 	FilterData = InArgs._FilterData;
 	InitialFilterData = FilterData;
 	bInEditMode = InArgs._InEditMode;
+	bEnableColorEditing = InArgs._InEnableColorEditing;
 	OnCreateFilter = InArgs._OnCreateFilter;
 	OnDeleteFilter = InArgs._OnDeleteFilter;
 	OnCancelClicked = InArgs._OnCancelClicked;
@@ -107,6 +108,7 @@ void SCustomTextFilterDialog::Construct( const FArguments& InArgs )
 				[
 					SNew(STextBlock)
 					.Text(LOCTEXT("FilterLabelText", "Filter Label"))
+					.ToolTipText(LOCTEXT("FilterLabelTooltip", "Display name of the filter"))
 				]
 				// Filter Name Input Box
 				+ SHorizontalBox::Slot()
@@ -114,7 +116,7 @@ void SCustomTextFilterDialog::Construct( const FArguments& InArgs )
 				.Padding(21.f, 0.f, 0.f, 0.f)
 				[
 					SNew(SBox)
-					.WidthOverride(254.f)
+					.WidthOverride(bEnableColorEditing ? 254.f : 556.f)
 					[
 						SAssignNew(FilterLabelTextBox, SEditableTextBox)
 						.Text_Lambda([this] ()
@@ -126,7 +128,6 @@ void SCustomTextFilterDialog::Construct( const FArguments& InArgs )
 							this->FilterData.FilterLabel = InText;
 						})
 					]
-					
 				]
 
 				// Filter Color Text
@@ -136,7 +137,9 @@ void SCustomTextFilterDialog::Construct( const FArguments& InArgs )
 				.Padding(21.f, 0.f, 0.f, 0.f)
 				[
 					SNew(STextBlock)
+					.Visibility(bEnableColorEditing ? EVisibility::Visible : EVisibility::Collapsed)
 					.Text(LOCTEXT("FilterColorText", "Color"))
+					.ToolTipText(LOCTEXT("FilterColorTooltip", "Display color of the filter button"))
 				]
 				// Filter Color
 				+ SHorizontalBox::Slot()
@@ -144,6 +147,7 @@ void SCustomTextFilterDialog::Construct( const FArguments& InArgs )
 				.Padding(21.f, 0.f, 0.f, 0.f)
 				[
 					SAssignNew(ColorBlock, SColorBlock )
+					.Visibility(bEnableColorEditing ? EVisibility::Visible : EVisibility::Collapsed)
 					.Color_Lambda([this] ()
 					{
 						return this->FilterData.FilterColor;
@@ -153,8 +157,8 @@ void SCustomTextFilterDialog::Construct( const FArguments& InArgs )
 					.AlphaDisplayMode(EColorBlockAlphaDisplayMode::Ignore)
 					.OnMouseButtonDown( this, &SCustomTextFilterDialog::ColorBlock_OnMouseButtonDown )
 				]
-				
 			]
+
 			// Filter String
 			+ SVerticalBox::Slot()
 			.AutoHeight()
@@ -169,6 +173,7 @@ void SCustomTextFilterDialog::Construct( const FArguments& InArgs )
 				[
 					SNew(STextBlock)
 					.Text(LOCTEXT("FilterStringText", "Filter String"))
+					.ToolTipText(LOCTEXT("FilterStringTooltip", "String literal to query for"))
 				]
 				// Filter Name Input Box
 				+ SHorizontalBox::Slot()
@@ -188,9 +193,9 @@ void SCustomTextFilterDialog::Construct( const FArguments& InArgs )
 							this->FilterData.FilterString = InText;
 						})
 					]
-					
 				]
 			]
+
 			// Buttons
 			+ SVerticalBox::Slot()
 			.VAlign(VAlign_Bottom)
diff --git a/Engine/Source/Developer/ToolWidgets/Public/Filters/CustomTextFilters.h b/Engine/Source/Developer/ToolWidgets/Public/Filters/CustomTextFilters.h
--- a/Engine/Source/Developer/ToolWidgets/Public/Filters/CustomTextFilters.h
+++ b/Engine/Source/Developer/ToolWidgets/Public/Filters/CustomTextFilters.h
@@ -7,6 +7,7 @@
 
 #include "CustomTextFilters.generated.h"
 
+#define LOCTEXT_NAMESPACE "CustomTextFilters"
 
 /* Struct containing the data that SCustomTextFilterDialog is currently editing */
 USTRUCT()
@@ -92,7 +93,7 @@ public:
 	/** Returns the tooltip for this filter, shown in the filters menu */
 	virtual FText GetToolTipText() const override
 	{
-		return TextFilter->GetRawFilterText(); // The tooltip will display the filter string
+		return FText::Format(LOCTEXT("CustomTextFiltersTooltip", "Filter by string \"{0}\""), TextFilter->GetRawFilterText()); // The tooltip will display the filter string
 	}
 
 	/** Returns the color this filter button will be when displayed as a button */
@@ -206,3 +207,5 @@ protected:
 	/* The Color of this filter pill */
 	FLinearColor Color;
 };
+
+#undef LOCTEXT_NAMESPACE
\ No newline at end of file
diff --git a/Engine/Source/Developer/ToolWidgets/Public/Filters/SBasicFilterBar.h b/Engine/Source/Developer/ToolWidgets/Public/Filters/SBasicFilterBar.h
--- a/Engine/Source/Developer/ToolWidgets/Public/Filters/SBasicFilterBar.h
+++ b/Engine/Source/Developer/ToolWidgets/Public/Filters/SBasicFilterBar.h
@@ -299,9 +299,9 @@ private:
 			OnFilterCtrlClicked = NewFilterCtrlClicked;
 		}
 
-		void SetOnFilterAltClicked(const FOnClicked& NewFilteAltClicked)
+		void SetOnFilterAltClicked(const FOnClicked& NewFilterAltClicked)
 		{
-			OnFilterAltClicked = NewFilteAltClicked;
+			OnFilterAltClicked = NewFilterAltClicked;
 		}
 
 		void SetOnFilterDoubleClicked( const FOnClicked& NewFilterDoubleClicked )
@@ -328,23 +328,26 @@ private:
 
 		virtual FReply OnMouseButtonUp( const FGeometry& InMyGeometry, const FPointerEvent& InMouseEvent ) override
 		{
+			FReply EventReply = FReply::Handled();
 			if (InMouseEvent.IsControlDown() && OnFilterCtrlClicked.IsBound())
 			{
-				return OnFilterCtrlClicked.Execute();
+				EventReply = OnFilterCtrlClicked.Execute();
 			}
 			else if (InMouseEvent.IsAltDown() && OnFilterAltClicked.IsBound())
 			{
-				return OnFilterAltClicked.Execute();
+				EventReply = OnFilterAltClicked.Execute();
 			}
 			else if( InMouseEvent.GetEffectingButton() == EKeys::MiddleMouseButton && OnFilterMiddleButtonClicked.IsBound() )
 			{
-				return OnFilterMiddleButtonClicked.Execute();
+				EventReply = OnFilterMiddleButtonClicked.Execute();
 			}
 			else
 			{
 				SCheckBox::OnMouseButtonUp(InMyGeometry, InMouseEvent);
-				return FReply::Handled().ReleaseMouseCapture();
 			}
+
+			bIsPressed = false;
+			return EventReply.ReleaseMouseCapture();
 		}
 
 	private:
@@ -1408,6 +1411,7 @@ protected:
 			SNew(SCustomTextFilterDialog)
 			.FilterData(CustomTextFilterData)
 			.InEditMode(bInEditMode)
+			.InEnableColorEditing(FilterPillStyle != EFilterPillStyle::Basic)
 			.OnCreateFilter(this, &SBasicFilterBar<FilterType>::OnCreateCustomTextFilter)
 			.OnDeleteFilter(this, &SBasicFilterBar<FilterType>::OnDeleteCustomTextFilter, InFilter)
 			.OnModifyFilter(this, &SBasicFilterBar<FilterType>::OnModifyCustomTextFilter, InFilter)
@@ -1471,50 +1475,55 @@ protected:
 			const TSharedRef<FFilterBase<FilterType>>& FrontendFilter = Filter->GetFilter().ToSharedRef();
 			
 			bool bIsFilterChecked = IsFrontendFilterInUse(FrontendFilter);
-			
-			TSharedPtr<SHorizontalBox> CustomTextFilterWidget = SNew(SHorizontalBox);
-			
-			// Checkbox for the filter
-			CustomTextFilterWidget->AddSlot()
-			.AutoWidth()
-			.Padding(16, 0, 0, 0)
-			[
-				SNew(SCheckBox)
-				.Style(FAppStyle::Get(), "Menu.CheckBox")
-				.OnCheckStateChanged(FOnCheckStateChanged::CreateSP(this, &SBasicFilterBar<FilterType>::CustomTextFilterClicked, Filter))
-				.IsChecked(bIsFilterChecked)
-			];
 
-			// A button showing the filter name, that checks it when clicked on
-			CustomTextFilterWidget->AddSlot()
-			.FillWidth(1.0)
-			[
+			TSharedPtr<SButton> CustomTextFilterWidget = 
 				SNew(SButton)
-				.ButtonStyle(FAppStyle::Get(), "SimpleButton")
-				.ForegroundColor(FStyleColors::White)
-				.Text(FrontendFilter->GetDisplayName())
+				.ButtonStyle(FAppStyle::Get(), "Menu.Button")
+				.ContentPadding(0.0f)
 				.OnClicked_Lambda([this, &Filter]()
 				{
 					this->FrontendFilterClicked(Filter->GetFilter().ToSharedRef());
 					FSlateApplication::Get().DismissAllMenus();
 					return FReply::Handled();
 				})
-			];
+				[
+					SNew(SHorizontalBox)
 
-			// An edit button to open the edit custom text filter dialog
-			CustomTextFilterWidget->AddSlot()
-			.AutoWidth()
-			.HAlign(HAlign_Right)
-			.Padding(0, 0, 16, 0)
-			[
-				SNew(SSimpleButton)
-				.Icon(FAppStyle::GetBrush("Icons.Edit"))
-				.OnClicked_Lambda([this, &Filter]()
-				{
-					this->CreateEditCustomTextFilterWindow(Filter);
-					return FReply::Handled();
-				})
-			];
+					// Checkbox for the filter
+					+ SHorizontalBox::Slot()
+					.AutoWidth()
+					.Padding(16, 0, 0, 0)
+					[
+						SNew(SCheckBox)
+						.Style(FAppStyle::Get(), "Menu.CheckBox")
+						.OnCheckStateChanged(FOnCheckStateChanged::CreateSP(this, &SBasicFilterBar<FilterType>::CustomTextFilterClicked, Filter))
+						.IsChecked(bIsFilterChecked)
+					]
+
+					// A text block showing the filter name
+					+ SHorizontalBox::Slot()
+					.FillWidth(1.0)
+					.Padding(4, 0, 0, 0)
+					.VAlign(VAlign_Center)
+					[
+						SNew(STextBlock)
+						.Text(FrontendFilter->GetDisplayName())
+					]
+
+					+ SHorizontalBox::Slot()
+					.AutoWidth()
+					.HAlign(HAlign_Right)
+					.Padding(0, 0, 16, 0)
+					[
+						SNew(SSimpleButton)
+						.Icon(FAppStyle::GetBrush("Icons.Edit"))
+						.OnClicked_Lambda([this, &Filter]()
+						{
+							this->CreateEditCustomTextFilterWindow(Filter);
+							return FReply::Handled();
+						})
+					]
+				];
 
 			CustomFiltersSection.AddEntry(
 			FToolMenuEntry::InitWidget(
diff --git a/Engine/Source/Developer/ToolWidgets/Public/Filters/SCustomTextFilterDialog.h b/Engine/Source/Developer/ToolWidgets/Public/Filters/SCustomTextFilterDialog.h
--- a/Engine/Source/Developer/ToolWidgets/Public/Filters/SCustomTextFilterDialog.h
+++ b/Engine/Source/Developer/ToolWidgets/Public/Filters/SCustomTextFilterDialog.h
@@ -20,13 +20,18 @@ public:
 	DECLARE_DELEGATE(FOnCancelClicked);
 	DECLARE_DELEGATE_OneParam(FOnGetFilterLabels, TArray<FText> & /* FilterNames */);
 
-	SLATE_BEGIN_ARGS(SCustomTextFilterDialog) {}
+	SLATE_BEGIN_ARGS(SCustomTextFilterDialog)
+		: _InEnableColorEditing(true)
+		{}
 	
     	/** The filter that this dialog is creating/editing */
     	SLATE_ARGUMENT(FCustomTextFilterData, FilterData)
 
 		/** True if we are editing an existing filter, false if we are creating a new one */
 		SLATE_ARGUMENT(bool, InEditMode)
+
+		/** True if we allow editing the filter color, mainly used to disable color editing on the basic pill style as it has no effect */
+		SLATE_ARGUMENT(bool, InEnableColorEditing)
 		
 		/** Delegate for when the Create button is clicked */
 		SLATE_EVENT(FOnCreateFilter, OnCreateFilter)
@@ -70,6 +75,9 @@ protected:
 	/* True if we are editing a filter, false if we are creating a new filter */
 	bool bInEditMode;
 
+	/* True if we want to give the user the ability to pick a color for the filter in the edit/create menu */
+	bool bEnableColorEditing;
+
 	/* The current filter data we are editing */
 	FCustomTextFilterData FilterData;
 	
