Fixed stack-use-after-scope caught by asan. Temp stack var bound to non const ref.

Introduced in CL 44671141.

#tests Editor with sidebar creation enabled/disabled

diff --git a/Engine/Source/Developer/ToolWidgets/Private/Sidebar/SSidebarDrawerContent.cpp b/Engine/Source/Developer/ToolWidgets/Private/Sidebar/SSidebarDrawerContent.cpp
--- a/Engine/Source/Developer/ToolWidgets/Private/Sidebar/SSidebarDrawerContent.cpp
+++ b/Engine/Source/Developer/ToolWidgets/Private/Sidebar/SSidebarDrawerContent.cpp
@@ -252,8 +252,8 @@ TArray<TSharedRef<ISidebarDrawerContent>> SSidebarDrawerContent::GetOrderedSecti
 
 void SSidebarDrawerContent::AddContentSlot(const TSharedRef<ISidebarDrawerContent>& InDrawerContent)
 {
-	SVerticalBox::FSlot::FSlotArguments& Slot = ContentBox->AddSlot()
-		.Padding(0.f, 0.f, 0.f, 2.f)
+	SVerticalBox::FScopedWidgetSlotArguments Slot = ContentBox->AddSlot();
+	Slot.Padding(0.f, 0.f, 0.f, 2.f)
 		[
 			SNew(SBox)
 			.Visibility(this, &SSidebarDrawerContent::GetSectionContentVisibility
