[Sequencer] Set default sidebar section fill height to 1.f

diff --git a/Engine/Source/Developer/ToolWidgets/Public/Sidebar/ISidebarDrawerContent.h b/Engine/Source/Developer/ToolWidgets/Public/Sidebar/ISidebarDrawerContent.h
--- a/Engine/Source/Developer/ToolWidgets/Public/Sidebar/ISidebarDrawerContent.h
+++ b/Engine/Source/Developer/ToolWidgets/Public/Sidebar/ISidebarDrawerContent.h
@@ -27,7 +27,7 @@ public:
 	virtual bool ShouldShowSection() const { return true; }
 
 	/** Sets the optional slot height stretch coefficient. If not set, auto sizes the slot. */
-	virtual TOptional<float> GetSectionFillHeight() const { return TOptional<float>(); }
+	virtual TOptional<float> GetSectionFillHeight() const { return 1.f; }
 
 	/** Gets the section sort order. Lower values will be put at the above higher values. */
 	virtual int32 GetSortOrder() const { return 0; }
