Avatar: If the description is bogus such as `[Unknown User]`, display a question mark instead of the first character.

#rnx

diff --git a/Engine/Source/Developer/ToolWidgets/Private/SAvatar.cpp b/Engine/Source/Developer/ToolWidgets/Private/SAvatar.cpp
--- a/Engine/Source/Developer/ToolWidgets/Private/SAvatar.cpp
+++ b/Engine/Source/Developer/ToolWidgets/Private/SAvatar.cpp
@@ -50,7 +50,8 @@ int32 SAvatar::OnPaint(const FPaintArgs& Args, const FGeometry& AllottedGeometry
 	{
 		FSlateFontInfo FontInfo = FCoreStyle::GetDefaultFontStyle("Bold", Radius);
 
-		const FString Text = Description.Left(1).ToUpper();
+		const FString Initial = Description.Left(1).ToUpper();
+		const FString Text = FChar::IsAlnum(Initial[0]) ? Initial : TEXT("?");
 		const FVector2D TextSize = FSlateApplication::Get().GetRenderer()->GetFontMeasureService()->Measure(Text, FontInfo);
 		const FVector2D TextOffset(Radius - TextSize.X / 2, Radius - TextSize.Y / 2);
 
diff --git a/Engine/Source/Developer/ToolWidgets/Public/SAvatar.h b/Engine/Source/Developer/ToolWidgets/Public/SAvatar.h
--- a/Engine/Source/Developer/ToolWidgets/Public/SAvatar.h
+++ b/Engine/Source/Developer/ToolWidgets/Public/SAvatar.h
@@ -9,7 +9,7 @@
 
 /** Widget for displaying a user avatar icon.
  * The avatar's background color is computed based on the Identifier.
- * The avatar's displayed letter is the initial letter of the Description.
+ * The avatar's displayed letter is the initial letter of the Description if it's alpha numeric, or a question mark otherwise.
  */
 class SAvatar : public SCompoundWidget
 {
