Deprecated SCustomDialog VAlignButtonBox since it didn't affect the button layout, but could result in the button rendering outside the dialog if VAlign_Center is used.
Updated the default value for SCustomDialog HAlignButtonBox to HAlign_Fill since that was the most likely behavior of it when it was uninitialized and it is needed to left align before button content by default (such as the message box expects).

diff --git a/Engine/Source/Developer/ToolWidgets/Private/Dialog/SCustomDialog.cpp b/Engine/Source/Developer/ToolWidgets/Private/Dialog/SCustomDialog.cpp
--- a/Engine/Source/Developer/ToolWidgets/Private/Dialog/SCustomDialog.cpp
+++ b/Engine/Source/Developer/ToolWidgets/Private/Dialog/SCustomDialog.cpp
@@ -73,7 +73,6 @@ void SCustomDialog::Construct(const FArguments& InArgs)
 	if (InArgs._Buttons.Num() > 0)
 	{
 		VerticalBox->AddSlot()
-			.VAlign(InArgs._VAlignButtonBox)
 			.HAlign(InArgs._HAlignButtonBox)
 			.AutoHeight()
 			.Padding(InArgs._ButtonAreaPadding)
diff --git a/Engine/Source/Developer/ToolWidgets/Public/Dialog/SCustomDialog.h b/Engine/Source/Developer/ToolWidgets/Public/Dialog/SCustomDialog.h
--- a/Engine/Source/Developer/ToolWidgets/Public/Dialog/SCustomDialog.h
+++ b/Engine/Source/Developer/ToolWidgets/Public/Dialog/SCustomDialog.h
@@ -96,8 +96,8 @@ public:
 		, _Icon(nullptr)
 		, _HAlignIcon(HAlign_Left)
 		, _VAlignIcon(VAlign_Center)
-		, _HAlignButtonBox(HAlign_Right)
-		, _VAlignButtonBox(VAlign_Center)
+		, _HAlignButtonBox(HAlign_Fill)
+		, _VAlignButtonBox(VAlign_Fill)
 		, _RootPadding(FMargin(4.f))
 		, _ButtonAreaPadding(FMargin(20.f, 16.f, 0.f, 0.f))
 		, _UseScrollBox(true)
@@ -147,11 +147,10 @@ public:
 		/** Custom widget placed before the auto-generated buttons */
 		SLATE_NAMED_SLOT(FArguments, BeforeButtons)
 
-		/** HAlign to use for Button Box slot (default: HAlign_Right) */
+		/** HAlign to use for Button Box slot (default: HAlign_Fill) */
 		SLATE_ARGUMENT(EHorizontalAlignment, HAlignButtonBox)
 
-		/** VAlign to use for Button Box  slot (default: VAlign_Center) */
-		SLATE_ARGUMENT(EVerticalAlignment, VAlignButtonBox)
+		SLATE_ARGUMENT_DEPRECATED(EVerticalAlignment, VAlignButtonBox, 5.8, "The VAlignButtonBox argument is deprecated and not in use anymore")
 
 		/** Padding to apply to the widget embedded in the window, i.e. to all widgets contained in the window (default: {4,4,4,4} )*/
 		SLATE_ATTRIBUTE(FMargin, RootPadding)
diff --git a/Engine/Source/Editor/LandscapeEditor/Private/LandscapeEditorDetailCustomization_NewLandscape.cpp b/Engine/Source/Editor/LandscapeEditor/Private/LandscapeEditorDetailCustomization_NewLandscape.cpp
--- a/Engine/Source/Editor/LandscapeEditor/Private/LandscapeEditorDetailCustomization_NewLandscape.cpp
+++ b/Engine/Source/Editor/LandscapeEditor/Private/LandscapeEditorDetailCustomization_NewLandscape.cpp
@@ -128,7 +128,6 @@ FReply FLandscapeEditorArrayBuilder_NewTargetLayer::ShowAutoFillTargetLayerDialo
 		const TSharedPtr<SCustomDialog> DialogWindow = SNew(SCustomDialog)
 			.Title(LOCTEXT("TargetLayerAutoFillLayers", "Auto-Fill Target Layer Assets"))
 			.HAlignButtonBox(HAlign_Center)
-			.VAlignButtonBox(VAlign_Center)
 			.Buttons({
 				SCustomDialog::FButton(
 					LOCTEXT("OK", "OK"),
diff --git a/Engine/Source/Editor/LandscapeEditor/Private/LandscapeEditorDetailCustomization_TargetLayers.cpp b/Engine/Source/Editor/LandscapeEditor/Private/LandscapeEditorDetailCustomization_TargetLayers.cpp
--- a/Engine/Source/Editor/LandscapeEditor/Private/LandscapeEditorDetailCustomization_TargetLayers.cpp
+++ b/Engine/Source/Editor/LandscapeEditor/Private/LandscapeEditorDetailCustomization_TargetLayers.cpp
@@ -593,7 +593,6 @@ FReply FLandscapeEditorCustomNodeBuilder_TargetLayers::ShowAutoFillTargetLayerDi
 	const TSharedPtr<SCustomDialog> DialogWindow = SNew(SCustomDialog)
 		.Title(LOCTEXT("TargetLayerAutoFillLayers", "Auto-Fill Target Layer Assets"))
 		.HAlignButtonBox(HAlign_Center)
-		.VAlignButtonBox(VAlign_Center)
 		.Buttons({
 			SCustomDialog::FButton(
 				LOCTEXT("Unassigned Layers Only", "Unassigned Layers Only"), 
diff --git a/Engine/Source/Editor/LandscapeEditor/Private/SLandscapeLayerListDialog.cpp b/Engine/Source/Editor/LandscapeEditor/Private/SLandscapeLayerListDialog.cpp
--- a/Engine/Source/Editor/LandscapeEditor/Private/SLandscapeLayerListDialog.cpp
+++ b/Engine/Source/Editor/LandscapeEditor/Private/SLandscapeLayerListDialog.cpp
@@ -203,7 +203,6 @@ void SLandscapeLayerListDialog::Construct(const FArguments& InArgs, const TWeakO
 		.Title(FText(LOCTEXT("LandscapeLayerListDialogTitleText", "Insert New Landscape Edit Layer")))
 		.UseScrollBox(false)
 		.HAlignButtonBox(HAlign_Center)
-		.VAlignButtonBox(VAlign_Center)
 		.WindowArguments(SWindow::FArguments()
 		.HasCloseButton(false))
 		.ContentAreaPadding(10.0f)
