Protect against crash reporting condition with widget drawer receiving a mouse up on a window that was already destroyed

#rnx
#tests forced path through FindWidget returning null to ensure it survives the early return

diff --git a/Engine/Source/Developer/ToolWidgets/Private/SWidgetDrawer.cpp b/Engine/Source/Developer/ToolWidgets/Private/SWidgetDrawer.cpp
--- a/Engine/Source/Developer/ToolWidgets/Private/SWidgetDrawer.cpp
+++ b/Engine/Source/Developer/ToolWidgets/Private/SWidgetDrawer.cpp
@@ -628,6 +628,12 @@ void SWidgetDrawer::OpenDrawer(const FName DrawerId)
 
 			TSharedPtr<SWindow> MyWindow = FSlateApplication::Get().FindWidgetWindow(AsShared());
 
+			// OpenDrawer is called on deferred mouse events so the window might already be closed by the time we get here
+			if (MyWindow == nullptr)
+			{
+				return;
+			}
+
 			const float MaxDrawerHeight = MyWindow->GetSizeInScreen().Y * 0.90f;
 
 			float TargetDrawerHeightPct = .33f;
