Handle unattended mode in SMessageDialog

- Added checks in `SCustomDialog::ShowModal` and `SCustomDialog::Show` to respect unattended mode.
- Added new argument in `SCustomDialog::ShowModal` to provide a default return value. Defaults to INDEX_NONE (-1). This matches current usage ie. the expected behavior is to return -1 by default on cancellation.
- Invoke `OnCancel` delegate in `SCustomDialog::ShowModal` and `SCustomDialog::Show` . This will mimic as close as possible the behaviour when the UI is available.

diff --git a/Engine/Source/Developer/ToolWidgets/Private/Dialog/SCustomDialog.cpp b/Engine/Source/Developer/ToolWidgets/Private/Dialog/SCustomDialog.cpp
--- a/Engine/Source/Developer/ToolWidgets/Private/Dialog/SCustomDialog.cpp
+++ b/Engine/Source/Developer/ToolWidgets/Private/Dialog/SCustomDialog.cpp
@@ -2,9 +2,13 @@
 
 #include "Dialog/SCustomDialog.h"
 
+#include "CoreGlobals.h"
 #include "Dialog/DialogCommands.h"
 #include "Framework/Application/SlateApplication.h"
 #include "Framework/Docking/TabManager.h"
+#include "Misc/App.h"
+#include "Misc/AutomationTest.h"
+#include "Misc/FeedbackContext.h"
 #include "Styling/AppStyle.h"
 #include "Styling/SlateBrush.h"
 #include "Widgets/Images/SImage.h"
@@ -82,10 +86,29 @@ void SCustomDialog::Construct(const FArguments& InArgs)
 	}
 }
 
-int32 SCustomDialog::ShowModal()
+int32 SCustomDialog::ShowModal(const int32 InDefaultReturnValue)
 {
-	FSlateApplication::Get().AddModalWindow(StaticCastSharedRef<SWindow>(this->AsShared()), FGlobalTabmanager::Get()->GetRootWindow());
-	return LastPressedButton;
+	int32 Result = InDefaultReturnValue;
+	if (!FApp::IsUnattended() && !GIsRunningUnattendedScript && FSlateApplication::IsInitialized())
+	{
+		FSlateApplication::Get().AddModalWindow(StaticCastSharedRef<SWindow>(this->AsShared()), FGlobalTabmanager::Get()->GetRootWindow());
+		Result = LastPressedButton;
+	}
+	else
+	{
+		OnClosed.ExecuteIfBound();
+		if (GWarn)
+		{
+			GWarn->Logf(TEXT("Message dialog was triggered in unattended script mode, title: %s"), *Title.Get().ToString());
+		}
+	}
+
+	if (GWarn)
+	{
+		GWarn->Logf(TEXT("Message dialog closed, result: %d, title: %s"), Result, *Title.Get().ToString());
+	}
+
+	return Result;
 }
 
 void SCustomDialog::SetOnCancelHotkeyPressed(const FSimpleDelegate& InOnCancelHotkeyPressed)
@@ -140,10 +163,26 @@ bool SCustomDialog::OnIsActiveChanged(const FWindowActivateEvent& ActivateEvent)
 
 void SCustomDialog::Show()
 {
-	const TSharedRef<SWindow> Window = FSlateApplication::Get().AddWindow(StaticCastSharedRef<SWindow>(this->AsShared()), true);
-	if (OnClosed.IsBound())
+	if (!FApp::IsUnattended() && !GIsRunningUnattendedScript && FSlateApplication::IsInitialized())
 	{
-		Window->GetOnWindowClosedEvent().AddLambda([this](const TSharedRef<SWindow>&) { OnClosed.Execute(); });
+		const TSharedRef<SWindow> Window = FSlateApplication::Get().AddWindow(StaticCastSharedRef<SWindow>(this->AsShared()), true);
+		if (OnClosed.IsBound())
+		{
+			Window->GetOnWindowClosedEvent().AddSPLambda(this, [this](const TSharedRef<SWindow>&) { OnClosed.Execute(); });
+		}
+	}
+	else
+	{
+		OnClosed.ExecuteIfBound();
+		if (GWarn)
+		{
+			GWarn->Logf(TEXT("Message dialog was triggered in unattended script mode, title: %s"), *Title.Get().ToString());
+		}
+	}
+
+	if (GWarn)
+	{
+		GWarn->Logf(TEXT("Message dialog closed, title: %s"), *Title.Get().ToString());
 	}
 }
 
@@ -331,3 +370,37 @@ void SCustomDialog::ExecuteConfirm()
 	// Try using the Confirm function
 	OnConfirmHotkeyPressed.ExecuteIfBound();
 }
+
+#if WITH_DEV_AUTOMATION_TESTS
+
+IMPLEMENT_SIMPLE_AUTOMATION_TEST(SCustomDialogShowModalTest, "System.ToolWidgets.CustomDialog.ShowModalInvokesOnClosedWhenUnattended",
+	EAutomationTestFlags::EditorContext | EAutomationTestFlags::CommandletContext | EAutomationTestFlags::EngineFilter)
+ 
+bool SCustomDialogShowModalTest::RunTest(const FString& Parameters)
+{
+	bool bClosed = false;
+	const TSharedRef<SCustomDialog> Dialog = SNew(SCustomDialog)
+		.Title(FText::FromString("Test dialog."))
+		.OnClosed_Lambda([&bClosed] {	bClosed = true;	});
+	
+	Dialog->ShowModal();
+		
+	return bClosed == true;
+}
+
+IMPLEMENT_SIMPLE_AUTOMATION_TEST(SCustomDialogShowTest, "System.ToolWidgets.CustomDialog.ShowInvokesOnClosedWhenUnattended",
+	EAutomationTestFlags::EditorContext | EAutomationTestFlags::CommandletContext | EAutomationTestFlags::EngineFilter)
+ 
+bool SCustomDialogShowTest::RunTest(const FString& Parameters)
+{
+	bool bClosed = false;
+	const TSharedRef<SCustomDialog> Dialog = SNew(SCustomDialog)
+		.Title(FText::FromString("Test dialog."))
+		.OnClosed_Lambda([&bClosed] {	bClosed = true;	});
+	
+	Dialog->Show();
+		
+	return bClosed == true;
+}
+
+#endif
\ No newline at end of file
diff --git a/Engine/Source/Developer/ToolWidgets/Public/Dialog/SCustomDialog.h b/Engine/Source/Developer/ToolWidgets/Public/Dialog/SCustomDialog.h
--- a/Engine/Source/Developer/ToolWidgets/Public/Dialog/SCustomDialog.h
+++ b/Engine/Source/Developer/ToolWidgets/Public/Dialog/SCustomDialog.h
@@ -206,9 +206,9 @@ public:
 	UE_API void Show();
 
 	/** Show a modal dialog. Will block until an input is received.
-	 * Returns the index of the button that was pressed.
+	 * Returns the index of the button that was pressed or InDefaultReturnValue if in unattended mode.
 	 */
-	UE_API int32 ShowModal();
+	UE_API int32 ShowModal(const int32 InDefaultReturnValue = INDEX_NONE);
 
 	/** Assigns a function to the Cancel hotkey */
 	UE_API void SetOnCancelHotkeyPressed(const FSimpleDelegate& InOnCancelHotkeyPressed);
