Return default value in SMessageDialog

- Return default value in `ShowModal` when the user closes the dialog with the X button.
- Added test to validate that `ShowModal` returns the default value in unattended mode.

diff --git a/Engine/Source/Developer/ToolWidgets/Private/Dialog/SCustomDialog.cpp b/Engine/Source/Developer/ToolWidgets/Private/Dialog/SCustomDialog.cpp
--- a/Engine/Source/Developer/ToolWidgets/Private/Dialog/SCustomDialog.cpp
+++ b/Engine/Source/Developer/ToolWidgets/Private/Dialog/SCustomDialog.cpp
@@ -92,7 +92,10 @@ int32 SCustomDialog::ShowModal(const int32 InDefaultReturnValue)
 	if (!FApp::IsUnattended() && !GIsRunningUnattendedScript && FSlateApplication::IsInitialized())
 	{
 		FSlateApplication::Get().AddModalWindow(StaticCastSharedRef<SWindow>(this->AsShared()), FGlobalTabmanager::Get()->GetRootWindow());
-		Result = LastPressedButton;
+		if (LastPressedButton != INDEX_NONE)
+		{
+			Result = LastPressedButton;
+		}
 	}
 	else
 	{
@@ -403,4 +406,16 @@ bool SCustomDialogShowTest::RunTest(const FString& Parameters)
 	return bClosed == true;
 }
 
+IMPLEMENT_SIMPLE_AUTOMATION_TEST(SCustomDialogShowModalReturnTest, "System.ToolWidgets.CustomDialog.ShowModalReturnsDefaultValueInUnattendedMode",
+	EAutomationTestFlags::EditorContext | EAutomationTestFlags::CommandletContext | EAutomationTestFlags::EngineFilter)
+ 
+bool SCustomDialogShowModalReturnTest::RunTest(const FString& Parameters)
+{
+	const TSharedRef<SCustomDialog> Dialog = SNew(SCustomDialog)
+		.Title(FText::FromString("Test dialog."));
+	
+	const int32 DefaultValue = 5;
+	return Dialog->ShowModal(DefaultValue) == DefaultValue;
+}
+
 #endif
\ No newline at end of file
diff --git a/Engine/Source/Developer/ToolWidgets/Public/Dialog/SCustomDialog.h b/Engine/Source/Developer/ToolWidgets/Public/Dialog/SCustomDialog.h
--- a/Engine/Source/Developer/ToolWidgets/Public/Dialog/SCustomDialog.h
+++ b/Engine/Source/Developer/ToolWidgets/Public/Dialog/SCustomDialog.h
@@ -238,7 +238,7 @@ private:
 	void ExecuteConfirm();
 	
 	/** The index of the button that was pressed last. */
-	int32 LastPressedButton = -1;
+	int32 LastPressedButton = INDEX_NONE;
 
 	FSimpleDelegate OnClosed;
 
