[ToolWidgets] Add conten's alignment to SWarningOrErrorBox.

#rnx

diff --git a/Engine/Source/Developer/ToolWidgets/Private/SWarningOrErrorBox.cpp b/Engine/Source/Developer/ToolWidgets/Private/SWarningOrErrorBox.cpp
--- a/Engine/Source/Developer/ToolWidgets/Private/SWarningOrErrorBox.cpp
+++ b/Engine/Source/Developer/ToolWidgets/Private/SWarningOrErrorBox.cpp
@@ -37,8 +37,8 @@ void SWarningOrErrorBox::Construct(const FArguments& InArgs)
 
 			+SHorizontalBox::Slot()
 			.AutoWidth()
-			.VAlign(VAlign_Center)
-			.HAlign(HAlign_Right)
+			.VAlign(InArgs._ContentVAlign)
+			.HAlign(InArgs._ContentHAlign)
 			[
 				InArgs._Content.Widget
 			]
diff --git a/Engine/Source/Developer/ToolWidgets/Public/SWarningOrErrorBox.h b/Engine/Source/Developer/ToolWidgets/Public/SWarningOrErrorBox.h
--- a/Engine/Source/Developer/ToolWidgets/Public/SWarningOrErrorBox.h
+++ b/Engine/Source/Developer/ToolWidgets/Public/SWarningOrErrorBox.h
@@ -21,6 +21,8 @@ public:
 		, _Padding(16.0f)
 		, _IconSize(24,24)
 		, _AutoWrapText(true)
+		, _ContentHAlign(EHorizontalAlignment::HAlign_Right)
+		, _ContentVAlign(EVerticalAlignment::VAlign_Center)
 		, _Content()
 	{}
 		SLATE_ATTRIBUTE(FText, Message)
@@ -29,6 +31,12 @@ public:
 		SLATE_ARGUMENT(FVector2D, IconSize)
 		SLATE_ARGUMENT(bool, AutoWrapText)
 
+		/** The Content horizontal alignment. */
+		SLATE_ARGUMENT(EHorizontalAlignment, ContentHAlign)
+
+		/** The Content vertical alignment. */
+		SLATE_ARGUMENT(EVerticalAlignment, ContentVAlign)
+
 		SLATE_DEFAULT_SLOT( FArguments, Content )
 
 	SLATE_END_ARGS()
