SMessageDialog : call SCustomDialog::OnKeyDown() if not handling the key in SMessageDialog::OnKeyDown()

This fixes the key bindings for the Confirm and Cancel commands.

diff --git a/Engine/Source/Developer/ToolWidgets/Private/Dialog/SMessageDialog.cpp b/Engine/Source/Developer/ToolWidgets/Private/Dialog/SMessageDialog.cpp
--- a/Engine/Source/Developer/ToolWidgets/Private/Dialog/SMessageDialog.cpp
+++ b/Engine/Source/Developer/ToolWidgets/Private/Dialog/SMessageDialog.cpp
@@ -85,8 +85,7 @@ FReply SMessageDialog::OnKeyDown(const FGeometry& MyGeometry, const FKeyEvent& I
 		return FReply::Handled();
 	}
 
-	//if it was some other button, ignore it
-	return FReply::Unhandled();
+	return SCustomDialog::OnKeyDown(MyGeometry, InKeyEvent);
 }
 
 FReply SMessageDialog::OnCopyMessage()
