SWidgetDrawer: Added ability to control visibility of the status bar drawer button.

#rnx

diff --git a/Engine/Source/Developer/ToolWidgets/Private/SWidgetDrawer.cpp b/Engine/Source/Developer/ToolWidgets/Private/SWidgetDrawer.cpp
--- a/Engine/Source/Developer/ToolWidgets/Private/SWidgetDrawer.cpp
+++ b/Engine/Source/Developer/ToolWidgets/Private/SWidgetDrawer.cpp
@@ -4,6 +4,7 @@
 #include "Widgets/Images/SImage.h"
 #include "Widgets/Input/SButton.h"
 #include "Widgets/Layout/SBorder.h"
+#include "Widgets/Layout/SBox.h"
 #include "Widgets/Text/SRichTextBlock.h"
 #include "Widgets/Text/STextBlock.h"
 #include "Widgets/Docking/SDockTab.h"
@@ -582,12 +583,20 @@ void SWidgetDrawer::RegisterDrawer(FWidgetDrawerConfig&& Drawer, int32 SlotIndex
 
 	if (RegisteredDrawers.Num() > NumDrawers)
 	{
-		TSharedRef<SWidget> Content = MakeStatusBarDrawerButton(Drawer);
+		TSharedRef<SWidget> Content = SNew(SBox)
+			.Padding(1.0f, 0.0f)
+			.Visibility_Lambda([Delegate = Drawer.GetDrawerButtonVisibilityDelegate, Name = Drawer.UniqueId]()
+				{
+					return Delegate.IsBound() ? Delegate.Execute(Name) : EVisibility::Visible;
+				}
+			)
+			[
+				MakeStatusBarDrawerButton(Drawer)
+			];
 
 		DrawerIdToContentWidget.Add(Drawer.UniqueId, Content);
 
 		DrawerBox->InsertSlot(SlotIndex)
-		.Padding(1.0f, 0.0f)
 		.AutoWidth()
 		[
 			Content
diff --git a/Engine/Source/Developer/ToolWidgets/Public/WidgetDrawerConfig.h b/Engine/Source/Developer/ToolWidgets/Public/WidgetDrawerConfig.h
--- a/Engine/Source/Developer/ToolWidgets/Public/WidgetDrawerConfig.h
+++ b/Engine/Source/Developer/ToolWidgets/Public/WidgetDrawerConfig.h
@@ -10,6 +10,7 @@ class SWidget;
 
 DECLARE_DELEGATE_OneParam(FOnStatusBarDrawerOpened, FName StatusBarName)
 DECLARE_DELEGATE_OneParam(FOnStatusBarDrawerDismissed, const TSharedPtr<SWidget>&)
+DECLARE_DELEGATE_RetVal_OneParam(EVisibility, FOnGetVisiblity, FName StatusBarName)
 
 struct FWidgetDrawerConfig
 {
@@ -19,6 +20,7 @@ struct FWidgetDrawerConfig
 
 	FName UniqueId;
 	FOnGetContent GetDrawerContentDelegate;
+	FOnGetVisiblity GetDrawerButtonVisibilityDelegate;
 	FOnStatusBarDrawerOpened OnDrawerOpenedDelegate;
 	FOnStatusBarDrawerDismissed OnDrawerDismissedDelegate;
 
