[ToolWidgets] Sidebar:
- Add SortOrder config property
- Sort by SortOrder, then by ButtonText
- Hide drawers that have no registered sections
#Synced-CL 49683375

diff --git a/Engine/Source/Developer/ToolWidgets/Private/Sidebar/SSidebar.cpp b/Engine/Source/Developer/ToolWidgets/Private/Sidebar/SSidebar.cpp
--- a/Engine/Source/Developer/ToolWidgets/Private/Sidebar/SSidebar.cpp
+++ b/Engine/Source/Developer/ToolWidgets/Private/Sidebar/SSidebar.cpp
@@ -40,20 +40,24 @@ void SSidebar::Construct(const FArguments& InArgs, const TSharedRef<SSidebarCont
 	SetVisibility(EVisibility::Visible);
 
 	ChildSlot
-	.Padding(FMargin(
-		TabLocation == ESidebarTabLocation::Right ? 2.f : 0.f,
-		TabLocation == ESidebarTabLocation::Bottom ? 2.f : 0.f,
-		TabLocation == ESidebarTabLocation::Left ? 2.f : 0.f,
-		TabLocation == ESidebarTabLocation::Top ? 2.f : 0.f))
 	[
-		SNew(SBorder)
-		.Padding(0.f)
-		.BorderImage(FAppStyle::Get().GetBrush(TEXT("Docking.Sidebar.Background")))
+		SNew(SBox)
+		.Padding(FMargin(
+			TabLocation == ESidebarTabLocation::Right ? 2.f : 0.f,
+			TabLocation == ESidebarTabLocation::Bottom ? 2.f : 0.f,
+			TabLocation == ESidebarTabLocation::Left ? 2.f : 0.f,
+			TabLocation == ESidebarTabLocation::Top ? 2.f : 0.f))
+		.Visibility(this, &SSidebar::GetSidebarVisibility)
 		[
-			SAssignNew(TabButtonContainer, SScrollBox)
-			.Orientation(IsHorizontal() ? Orient_Horizontal : Orient_Vertical)
-			.ScrollBarAlwaysVisible(false)
-			.ScrollBarVisibility(EVisibility::Collapsed)
+			SNew(SBorder)
+			.Padding(0.f)
+			.BorderImage(FAppStyle::Get().GetBrush(TEXT("Docking.Sidebar.Background")))
+			[
+				SAssignNew(TabButtonContainer, SScrollBox)
+				.Orientation(IsHorizontal() ? Orient_Horizontal : Orient_Vertical)
+				.ScrollBarAlwaysVisible(false)
+				.ScrollBarVisibility(EVisibility::Collapsed)
+			]
 		]
 	];
 }
@@ -72,21 +76,19 @@ bool SSidebar::RegisterDrawer(FSidebarDrawerConfig&& InDrawerConfig)
 	NewDrawer->ContentWidget = NewDrawer->Config.OverrideContentWidget.IsValid()
 		? NewDrawer->Config.OverrideContentWidget
 		: SNew(SSidebarDrawerContent, NewDrawer);
+	Drawers.Add(NewDrawer);
 
-	// Add tab button
-	TabButtonContainer->AddSlot()
-		[
-			SAssignNew(NewDrawer->ButtonWidget, SSidebarButton, NewDrawer, GetTabLocation())
-			.MinButtonSize(bAlwaysUseMaxButtonSize ? MaxTabButtonSize : MinTabButtonSize)
-			.MaxButtonSize(MaxTabButtonSize)
-			.ButtonThickness(TabButtonThickness)
-			.OnPressed(this, &SSidebar::OnTabDrawerButtonPressed)
-			.OnPinToggled(this, &SSidebar::OnDrawerTabPinToggled)
-			.OnDockToggled(this, &SSidebar::OnDrawerTabDockToggled)
-			.OnGetContextMenuContent(this, &SSidebar::OnGetTabDrawerContextMenuWidget, NewDrawer)
-		];
+	Drawers.Sort([](const TSharedRef<FSidebarDrawer>& InA, const TSharedRef<FSidebarDrawer>& InB) -> bool
+		{
+			if (InA->Config.SortOrder == InB->Config.SortOrder)
+			{
+				return InA->Config.ButtonText.Get(FText())
+					.CompareTo(InB->Config.ButtonText.Get(FText())) < 0;
+			}
+			return InA->Config.SortOrder < InB->Config.SortOrder;
+		});
 
-	Drawers.Add(NewDrawer);
+	RecreateDrawerWidgets();
 
 	const FName DrawerId = NewDrawer->GetUniqueId();
 
@@ -166,6 +168,47 @@ int32 SSidebar::GetDrawerCount() const
 	return Drawers.Num();
 }
 
+EVisibility SSidebar::GetSidebarVisibility() const
+{
+	// There may be registered drawers that are hidden due to not having any sections,
+	// so we cannot just check if Drawers is empty.
+	for (const TSharedRef<FSidebarDrawer>& Drawer : Drawers)
+	{
+		if (Drawer->ContentSections.Num() > 0)
+		{
+			return EVisibility::Visible;
+		}
+	}
+	return EVisibility::Collapsed;
+}
+
+void SSidebar::RecreateDrawerWidgets()
+{
+	TabButtonContainer->ClearChildren();
+
+	for (const TSharedRef<FSidebarDrawer>& Drawer : Drawers)
+	{
+		const TWeakPtr<FSidebarDrawer> WeakDrawer = Drawer;
+
+		TabButtonContainer->AddSlot()
+			[
+				SAssignNew(Drawer->ButtonWidget, SSidebarButton, Drawer, GetTabLocation())
+				.Visibility_Lambda([WeakDrawer]()
+					{
+						return WeakDrawer.IsValid() && WeakDrawer.Pin()->ContentSections.IsEmpty()
+							? EVisibility::Collapsed : EVisibility::Visible;
+					})
+				.MinButtonSize(bAlwaysUseMaxButtonSize ? MaxTabButtonSize : MinTabButtonSize)
+				.MaxButtonSize(MaxTabButtonSize)
+				.ButtonThickness(TabButtonThickness)
+				.OnPressed(this, &SSidebar::OnTabDrawerButtonPressed)
+				.OnPinToggled(this, &SSidebar::OnDrawerTabPinToggled)
+				.OnDockToggled(this, &SSidebar::OnDrawerTabDockToggled)
+				.OnGetContextMenuContent(this, &SSidebar::OnGetTabDrawerContextMenuWidget, Drawer)
+			];
+	}
+}
+
 bool SSidebar::RegisterDrawerSection(const FName InDrawerId, const TSharedPtr<ISidebarDrawerContent>& InSection)
 {
 	const TSharedPtr<FSidebarDrawer> Drawer = FindDrawer(InDrawerId);
diff --git a/Engine/Source/Developer/ToolWidgets/Public/Sidebar/SSidebar.h b/Engine/Source/Developer/ToolWidgets/Public/Sidebar/SSidebar.h
--- a/Engine/Source/Developer/ToolWidgets/Public/Sidebar/SSidebar.h
+++ b/Engine/Source/Developer/ToolWidgets/Public/Sidebar/SSidebar.h
@@ -244,6 +244,10 @@ public:
 	UE_API void RebuildDrawer(const FName InDrawerId, const bool bInOnlyIfOpen = true);
 
 private:
+	EVisibility GetSidebarVisibility() const;
+
+	void RecreateDrawerWidgets();
+
 	UE_API void OnTabDrawerButtonPressed(const TSharedRef<FSidebarDrawer>& InDrawer);
 	UE_API void OnDrawerTabPinToggled(const TSharedRef<FSidebarDrawer>& InDrawer, const bool bIsPinned);
 	UE_API void OnDrawerTabDockToggled(const TSharedRef<FSidebarDrawer>& InDrawer, const bool bIsDocked);
diff --git a/Engine/Source/Developer/ToolWidgets/Public/Sidebar/SidebarDrawerConfig.h b/Engine/Source/Developer/ToolWidgets/Public/Sidebar/SidebarDrawerConfig.h
--- a/Engine/Source/Developer/ToolWidgets/Public/Sidebar/SidebarDrawerConfig.h
+++ b/Engine/Source/Developer/ToolWidgets/Public/Sidebar/SidebarDrawerConfig.h
@@ -38,6 +38,9 @@ struct FSidebarDrawerConfig
 	/** Unique Id used to identify the drawer */
 	FName UniqueId;
 
+	/** Primary sort order of the drawer tab button. Secondary sort by ButtonText property. */
+	float SortOrder = 10.f;
+
 	/** Text to display on the drawer tab button */
 	TAttribute<FText> ButtonText;
 
