don't use clipboard icon in message dialog when it is used outside of the editor

#rnx

#Synced-CL 50053775

diff --git a/Engine/Source/Developer/ToolWidgets/Private/Dialog/SMessageDialog.cpp b/Engine/Source/Developer/ToolWidgets/Private/Dialog/SMessageDialog.cpp
--- a/Engine/Source/Developer/ToolWidgets/Private/Dialog/SMessageDialog.cpp
+++ b/Engine/Source/Developer/ToolWidgets/Private/Dialog/SMessageDialog.cpp
@@ -8,8 +8,12 @@
 #include "Styling/SlateBrush.h"
 #include "Widgets/Text/SRichTextBlock.h"
 #include "Widgets/Text/STextBlock.h"
+#if WITH_EDITOR
 #include "Widgets/Input/SButton.h"
 #include "Widgets/Images/SImage.h"
+#else
+#include "Widgets/Input/SHyperlink.h"
+#endif
 
 void SMessageDialog::Construct(const FArguments& InArgs)
 {
@@ -39,6 +43,8 @@ void SMessageDialog::Construct(const FArguments& InArgs)
 			.WrapTextAt(InArgs._WrapMessageAt);
 	}
 
+	const FText CopyToolTipText = NSLOCTEXT("SChoiceDialog", "CopyMessageTooltip", "Copy the text in this message to the clipboard (CTRL+C)");
+
 	SCustomDialog::Construct(SCustomDialog::FArguments()
 		.Title(InArgs._Title)
 		.Content()
@@ -62,10 +68,11 @@ void SMessageDialog::Construct(const FArguments& InArgs)
 		.OnClosed(InArgs._OnClosed)
 		.BeforeButtons()
 		[
+#if WITH_EDITOR	
 			SNew(SButton)
 			.ButtonStyle(FAppStyle::Get(), "SimpleButton")
 			.OnClicked(this, &SMessageDialog::OnCopyMessage)
-			.ToolTipText(NSLOCTEXT("SChoiceDialog", "CopyMessageTooltip", "Copy the text in this message to the clipboard (CTRL+C)"))
+			.ToolTipText(CopyToolTipText)
 			.ContentPadding(2.f)
 			.Content()
 			[
@@ -73,6 +80,12 @@ void SMessageDialog::Construct(const FArguments& InArgs)
 				.Image(FAppStyle::Get().GetBrush("Icons.Clipboard"))
 				.ColorAndOpacity(FSlateColor::UseForeground())
 			]
+#else
+			SNew(SHyperlink)
+			.OnNavigate_Lambda( [this]() { OnCopyMessage(); } ) // using a lambda here to avoid changing the signature of OnCopyMessage
+			.Text( NSLOCTEXT("SChoiceDialog", "CopyMessageHyperlink", "Copy") )
+			.ToolTipText(CopyToolTipText)
+#endif
 		]
 	);
 }
