[ToolWidgets] Add SPopoutTabInlineContent widget that allows inline content to be popped out as a tab window containing the content
- Will be used to enable popout of transport controls
- Eventually will be used to replace EaseCurveTool popout functionality
- May be used my other widgets to enable popping out inline content

#Synced-CL 50365487

diff --git a/Engine/Source/Developer/ToolWidgets/Private/SPopoutTabInlineContent.cpp b/Engine/Source/Developer/ToolWidgets/Private/SPopoutTabInlineContent.cpp
new file mode 100644
--- /dev/null
+++ b/Engine/Source/Developer/ToolWidgets/Private/SPopoutTabInlineContent.cpp
@@ -0,0 +1,680 @@
+﻿// Copyright Epic Games, Inc. All Rights Reserved.
+
+#include "SPopoutTabInlineContent.h"
+#include "Brushes/SlateRoundedBoxBrush.h"
+#include "Framework/Application/SlateApplication.h"
+#include "Framework/Docking/TabManager.h"
+#include "Framework/MultiBox/MultiBoxBuilder.h"
+#include "InputCoreTypes.h"
+#include "ToolMenu.h"
+#include "ToolMenus.h"
+#include "Styling/StyleColors.h"
+#include "Widgets/Images/SImage.h"
+#include "Widgets/Input/SButton.h"
+#include "Widgets/Layout/SBox.h"
+#include "Widgets/SOverlay.h"
+#include "Widgets/SWindow.h"
+
+#define LOCTEXT_NAMESPACE "SPopoutTabInlineContent"
+
+SPopoutTabInlineContent::~SPopoutTabInlineContent()
+{
+	UnsubscribeFromTabForegrounded();
+
+	UnregisterAutoSizeTimer();
+
+	// Avoid complex tab operations during engine shutdown
+	if (!IsEngineExitRequested() && FSlateApplication::IsInitialized())
+	{
+		// Request the tab to close. OnTabClosed will restore safely if it actually closes.
+		CloseTab();
+		UnregisterTab();
+
+		// If there was no tab or we're in teardown and it won't close, this will restore if detached
+		if (!WeakSpawnedDockTab.IsValid())
+		{
+			RestoreInlineContent();
+		}
+	}
+}
+
+void SPopoutTabInlineContent::Construct(const FArguments& InArgs)
+{
+	WeakContentWidget = InArgs._Content.Widget;
+	WeakTabManager = InArgs._TabManager.IsValid() ? InArgs._TabManager.Pin().ToSharedRef() : FGlobalTabmanager::Get();
+	TabId = InArgs._TabId;
+	TabDisplayName = InArgs._TabDisplayName;
+	TabIcon = InArgs._TabIcon;
+	TabRole = InArgs._TabRole;
+	WeakTabGroup = InArgs._TabGroup;
+	bShouldAutosize = InArgs._ShouldAutosize;
+	bCanDragToPopout = InArgs._CanDragToPopout;
+	bShowHandle = InArgs._ShowHandle;
+	bHideDockTabStackTab = InArgs._HideDockTabStackTab;
+	PopoutStateChangedDelegate = InArgs._OnPopoutStateChanged;
+	GetMenuContentDelegate = InArgs._OnGetMenuContent;
+
+	static const FSlateRoundedBoxBrush BorderHandleBrush(FLinearColor::Transparent
+		, 0.f, FStyleColors::AccentBlue, 2.f);
+
+	ChildSlot
+	.HAlign(HAlign_Fill)
+	.VAlign(VAlign_Fill)
+	[
+		SNew(SBox)
+		.HAlign(HAlign_Center)
+		.VAlign(VAlign_Center)
+		.Visibility(this, &SPopoutTabInlineContent::GetContentVisibility)
+		[
+			SNew(SBox)
+			.HAlign(InArgs._ContentHAlign)
+			.VAlign(InArgs._ContentVAlign)
+			[
+				SNew(SOverlay)
+				+ SOverlay::Slot()
+				[
+					SAssignNew(InlineContentContainer, SBox)
+					[
+						InArgs._Content.Widget
+					]
+				]
+				+ SOverlay::Slot()
+				[
+					SNew(SImage)
+					.Image(&BorderHandleBrush)
+					.Visibility(this, &SPopoutTabInlineContent::GetContentBorderVisibility)
+				]
+			]
+		]
+	];
+
+	if (InArgs._InitiallyPoppedOut)
+	{
+		Popout(true);
+
+		RegisterAutoSizeTimer();
+	}
+}
+
+TSharedRef<SButton> SPopoutTabInlineContent::CreateToggleButton(const TSharedRef<SPopoutTabInlineContent>& InOwner
+	, const FText& InTooltipText
+	, const FSlateBrush* const InIconBrush)
+{
+	constexpr float ToolButtonImageSize = 12.f;
+
+	const TWeakPtr<SPopoutTabInlineContent> WeakOwner = InOwner;
+
+	return SNew(SButton)
+		.VAlign(VAlign_Center)
+		.ToolTipText(InTooltipText)
+		.Visibility_Lambda([WeakOwner]()
+			{
+				if (const TSharedPtr<SPopoutTabInlineContent> Owner = WeakOwner.Pin())
+				{
+					return Owner->IsPoppedOut() ? EVisibility::Collapsed : EVisibility::Visible;
+				}
+				return EVisibility::Collapsed;
+			})
+		.OnClicked_Lambda([WeakOwner]()
+			{
+				if (const TSharedPtr<SPopoutTabInlineContent> Owner = WeakOwner.Pin())
+				{
+					Owner->Popout(true);
+				}
+				return FReply::Handled();
+			})
+		[
+			SNew(SImage)
+			.DesiredSizeOverride(FVector2D(ToolButtonImageSize, ToolButtonImageSize))
+			.ColorAndOpacity(FSlateColor::UseForeground())
+			.Image(InIconBrush)
+		];
+}
+
+bool SPopoutTabInlineContent::IsPoppedOut() const
+{
+	return bPoppedOut;
+}
+
+void SPopoutTabInlineContent::Popout(const bool bInPopout)
+{
+	if (bInPopout)
+	{
+		RegisterTab();
+
+		if (const TSharedPtr<FTabManager> TabManager = GetTabManager())
+		{
+			if (const TSharedPtr<SDockTab> Tab = TabManager->TryInvokeTab(TabId))
+			{
+				BroadcastPopoutStateChanged(true);
+			}
+			else
+			{
+				UnregisterTab();
+				BroadcastPopoutStateChanged(false);
+			}
+		}
+	}
+	else
+	{
+		CloseTab();
+	}
+
+	bLeftMouseDownOverWidget = false;
+}
+
+void SPopoutTabInlineContent::TogglePopout()
+{
+	Popout(!IsPoppedOut());
+}
+
+TSharedPtr<FTabManager> SPopoutTabInlineContent::GetTabManager() const
+{
+	return WeakTabManager.Pin();
+}
+
+bool SPopoutTabInlineContent::IsUniqueTabId() const
+{
+	if (TabId.IsNone())
+	{
+		return false;
+	}
+
+	const TSharedPtr<FTabManager> TabManager = GetTabManager();
+	if (!TabManager.IsValid())
+	{
+		return false;
+	}
+
+	// If anyone already registered a spawner under this id, we might hijack their tab
+	if (TabManager->HasTabSpawner(TabId))
+	{
+		return false;
+	}
+
+	// If a live tab already exists under this id, we might be pointing at someone else’s instance
+	if (TabManager->FindExistingLiveTab(TabId).IsValid())
+	{
+		return false;
+	}
+
+	return true;
+}
+
+void SPopoutTabInlineContent::RegisterTab()
+{
+	if (bTabRegistered)
+	{
+		return;
+	}
+
+	if (!IsUniqueTabId())
+	{
+		UE_LOG(LogTemp, Warning
+			, TEXT("SPopoutTabInlineContent: TabId '%s' conflicts with an existing spawner or live tab; refusing to register.")
+			, *TabId.ToString());
+		return;
+	}
+
+	const TSharedPtr<FTabManager> TabManager = GetTabManager();
+	if (!TabManager.IsValid())
+	{
+		return;
+	}
+
+	FTabSpawnerEntry& TabSpawnerEntry = TabManager->RegisterTabSpawner(TabId
+		, FOnSpawnTab::CreateSP(this, &SPopoutTabInlineContent::SpawnTab))
+			.SetDisplayName(TabDisplayName)
+			.SetIcon(TabIcon);
+
+	if (const TSharedPtr<FWorkspaceItem> TabGroup = WeakTabGroup.Pin())
+	{
+		TabSpawnerEntry.SetGroup(TabGroup.ToSharedRef());
+	}
+
+	bTabRegistered = true;
+}
+
+void SPopoutTabInlineContent::UnregisterTab()
+{
+	if (!bTabRegistered)
+	{
+		return;
+	}
+
+	if (const TSharedPtr<FTabManager> TabManager = GetTabManager())
+	{
+		TabManager->UnregisterTabSpawner(TabId);
+	}
+
+	bTabRegistered = false;
+}
+
+TSharedRef<SDockTab> SPopoutTabInlineContent::SpawnTab(const FSpawnTabArgs& InArgs)
+{
+	const TSharedPtr<SWidget> ContentWidget = WeakContentWidget.IsValid()
+		? WeakContentWidget.Pin() : DetachedContentWidget;
+
+	// If we have the real widget, remove it from the inline container before reparenting
+	if (ContentWidget.IsValid() && InlineContentContainer.IsValid())
+	{
+		DetachedContentWidget = ContentWidget;
+		InlineContentContainer->SetContent(SNullWidget::NullWidget);
+	}
+
+	const TSharedRef<SWidget> ActualContentWidget = DetachedContentWidget.IsValid()
+		? DetachedContentWidget.ToSharedRef() : SNullWidget::NullWidget;
+
+	const TSharedRef<SDockTab> DockTab = SNew(SDockTab)
+		.TabRole(TabRole)
+		.ShouldAutosize(bShouldAutosize)
+		.OnTabActivated(this, &SPopoutTabInlineContent::OnTabActivated)
+		.OnTabClosed(this, &SPopoutTabInlineContent::OnTabClosed)
+		.OnTabRelocated(this, &SPopoutTabInlineContent::OnTabRelocated)
+		[
+			ActualContentWidget
+		];
+	WeakSpawnedDockTab = DockTab;
+
+	RegisterAutoSizeTimer(DockTab);
+
+	SubscribeToTabForegrounded();
+
+	BroadcastPopoutStateChanged(true);
+
+	return DockTab;
+}
+
+void SPopoutTabInlineContent::CloseTab()
+{
+	if (const TSharedPtr<SDockTab> ExistingTab = WeakSpawnedDockTab.Pin())
+	{
+		// Do NOT restore here. The tab still owns the widget.
+		// OnTabClosed will call RestoreInlineContent() safely.
+		ExistingTab->RequestCloseTab();
+
+		return;
+	}
+
+	// If we don't have a spawned tab reference, we can't safely assume a tab isn't open.
+	// Best-effort: restore only if we are detached and there is no known live tab mechanism.
+	RestoreInlineContent();
+
+	BroadcastPopoutStateChanged(false);
+}
+
+void SPopoutTabInlineContent::OnTabActivated(const TSharedRef<SDockTab> InDockTab, const ETabActivationCause InActivationCause)
+{
+	RegisterAutoSizeTimer(InDockTab);
+
+	BroadcastPopoutStateChanged(true);
+}
+
+void SPopoutTabInlineContent::OnTabClosed(const TSharedRef<SDockTab> InDockTab)
+{
+	UnregisterAutoSizeTimer();
+
+	// Restore only once the tab is actually closed so the widget isn't double-parented
+	WeakSpawnedDockTab.Reset(); // reset first so RestoreInlineContent won't bail
+	RestoreInlineContent();
+
+	UnregisterTab();
+
+	UnsubscribeFromTabForegrounded();
+
+	bIsForeground = false;
+
+	BroadcastPopoutStateChanged(false);
+}
+
+void SPopoutTabInlineContent::OnTabRelocated()
+{
+	TryHideStackTab();
+
+	RegisterAutoSizeTimer();
+}
+
+void SPopoutTabInlineContent::OnTabForegrounded(const TSharedPtr<SDockTab> InNewTab, const TSharedPtr<SDockTab> InPreviousTab)
+{
+	if (!InNewTab.IsValid())
+	{
+		return;
+	}
+
+	const TSharedPtr<SDockTab> SpawnedTab = WeakSpawnedDockTab.Pin();
+
+	const auto IsOurTab = [&SpawnedTab](const TSharedPtr<SDockTab>& Tab) -> bool
+		{
+			return Tab.IsValid() && (Tab == SpawnedTab);
+		};
+
+	const bool bNewIsOurs = IsOurTab(InNewTab);
+	const bool bPrevIsOurs = IsOurTab(InPreviousTab);
+
+	if (!bNewIsOurs)
+	{
+		bIsForeground = false;
+		return;
+	}
+
+	const bool bBecameForeground = !bPrevIsOurs || !bIsForeground;
+	bIsForeground = true;
+
+	if (bBecameForeground)
+	{
+		TryHideStackTab();
+	}
+}
+
+FCursorReply SPopoutTabInlineContent::OnCursorQuery(const FGeometry& InGeometry, const FPointerEvent& InPointerEvent) const
+{
+	if (bLeftMouseDownOverWidget && bCanDragToPopout)
+	{
+		return FCursorReply::Cursor(EMouseCursor::GrabHand);
+	}
+	return FCursorReply::Unhandled();
+}
+
+FReply SPopoutTabInlineContent::OnMouseButtonDown(const FGeometry& InGeometry, const FPointerEvent& InPointerEvent)
+{
+	const FKey EffectingButton = InPointerEvent.GetEffectingButton();
+	const bool bLeftMouseDown = EffectingButton == EKeys::LeftMouseButton;
+	const bool bRightMouseDown = EffectingButton == EKeys::RightMouseButton;
+
+	if (bLeftMouseDown)
+	{
+		if (bCanDragToPopout && !IsPoppedOut())
+		{
+			bLeftMouseDownOverWidget = true;
+
+			// Keep receiving move events even if the cursor leaves the widget while dragging
+			return FReply::Handled()
+				.CaptureMouse(SharedThis(this))
+				.DetectDrag(SharedThis(this), EKeys::LeftMouseButton);
+		}
+	}
+	else if (bRightMouseDown)
+	{
+		bRightMouseDownOverWidget = true;
+	}
+
+	return FReply::Unhandled();
+}
+
+FReply SPopoutTabInlineContent::OnMouseButtonUp(const FGeometry& InGeometry, const FPointerEvent& InPointerEvent)
+{
+	const FKey EffectingButton = InPointerEvent.GetEffectingButton();
+	const bool bLeftMouseUp = EffectingButton == EKeys::LeftMouseButton;
+	const bool bRightMouseUp = EffectingButton == EKeys::RightMouseButton;
+
+	if (bLeftMouseUp)
+	{
+		bLeftMouseDownOverWidget = false;
+
+		if (HasMouseCapture())
+		{
+			return FReply::Handled().ReleaseMouseCapture();
+		}
+	}
+
+	if (bRightMouseUp)
+	{
+		if (bRightMouseDownOverWidget)
+		{
+			bRightMouseDownOverWidget = false;
+
+			// Show the context menu
+			const TSharedRef<SWidget> ParentWidget = SharedThis(this);
+			const FWidgetPath* const EventPath = InPointerEvent.GetEventPath();
+			const FWidgetPath WidgetPath = EventPath ? *EventPath : FWidgetPath();
+			const TSharedRef<SWidget> MenuContentWidget = GetMenuContentDelegate.IsBound()
+				? GetMenuContentDelegate.Execute() : CreateDefaultMenuContent();
+
+			FSlateApplication::Get().PushMenu(ParentWidget
+				, WidgetPath
+				, MenuContentWidget
+				, InPointerEvent.GetScreenSpacePosition()
+				, FPopupTransitionEffect(FPopupTransitionEffect::ContextMenu)
+			);
+
+			return FReply::Handled();
+		}
+	}
+
+	return FReply::Unhandled();
+}
+
+void SPopoutTabInlineContent::OnMouseLeave(const FPointerEvent& InPointerEvent)
+{
+	// Don't clear drag state while the left button is held. Leaving the widget during a drag is normal.
+	if (!InPointerEvent.IsMouseButtonDown(EKeys::LeftMouseButton))
+	{
+		bLeftMouseDownOverWidget = false;
+	}
+
+	bRightMouseDownOverWidget = false;
+}
+
+FReply SPopoutTabInlineContent::OnDragDetected(const FGeometry& InGeometry, const FPointerEvent& InPointerEvent)
+{
+	if (InPointerEvent.IsMouseButtonDown(EKeys::LeftMouseButton)
+		&& bCanDragToPopout
+		&& !IsPoppedOut())
+	{
+		Popout(true);
+	}
+
+	bLeftMouseDownOverWidget = false;
+
+	// We’re not starting a drag-drop operation, just using drag as a gesture
+	return FReply::Handled().ReleaseMouseCapture();
+}
+
+TSharedRef<SWidget> SPopoutTabInlineContent::CreateDefaultMenuContent()
+{
+	FMenuBuilder MenuBuilder(true, nullptr);
+
+	MenuBuilder.BeginSection(TEXT("ContentOptions"), LOCTEXT("ContentOptions", "Options"));
+
+	MenuBuilder.AddMenuEntry(LOCTEXT("PopoutContent", "Popout")
+		, LOCTEXT("PopoutContentTooltip", "Popout this content to its own tab window")
+		, FSlateIcon()
+		, FExecuteAction::CreateSP(this, &SPopoutTabInlineContent::Popout, true)
+	);
+
+	MenuBuilder.EndSection();
+
+	return MenuBuilder.MakeWidget();
+}
+
+EVisibility SPopoutTabInlineContent::GetContentVisibility() const
+{
+	return IsPoppedOut() ? EVisibility::Collapsed : EVisibility::Visible;
+}
+
+EVisibility SPopoutTabInlineContent::GetContentBorderVisibility() const
+{
+	return (bShowHandle && bLeftMouseDownOverWidget && bCanDragToPopout)
+		? EVisibility::Visible : EVisibility::Hidden;
+}
+
+void SPopoutTabInlineContent::RegisterAutoSizeTimer(const TSharedPtr<SDockTab>& InDockTab)
+{
+	// This function and the timer it creates is only used for auto-sizing
+	if (!bShouldAutosize)
+	{
+		return;
+	}
+
+	if (SetupWindowTimerHandle.IsValid())
+	{
+		return;
+	}
+
+	const TSharedPtr<SDockTab> ExistingTab = InDockTab.IsValid() ? InDockTab : WeakSpawnedDockTab.Pin();
+	if (!ExistingTab.IsValid())
+	{
+		return;
+	}
+
+	const TWeakPtr<SPopoutTabInlineContent> ThisWeak = SharedThis(this);
+	const TWeakPtr<SDockTab> WeakTab = ExistingTab;
+
+	/** Need to use ExistingTab to register the timer since this widget can be hidden */
+	SetupWindowTimerHandle = ExistingTab->RegisterActiveTimer(0.f
+		, FWidgetActiveTimerDelegate::CreateLambda(
+			[ThisWeak, WeakTab](const double InCurrentTime, const float InDeltaTime) -> EActiveTimerReturnType
+			{
+				if (const TSharedPtr<SPopoutTabInlineContent> This = ThisWeak.Pin())
+				{
+					return This->HandleAutoSizeWindowTimer(InCurrentTime, InDeltaTime, WeakTab);
+				}
+				return EActiveTimerReturnType::Stop;
+			}));
+}
+
+void SPopoutTabInlineContent::UnregisterAutoSizeTimer()
+{
+	if (!SetupWindowTimerHandle.IsValid())
+	{
+		return;
+	}
+
+	if (const TSharedPtr<SDockTab> ExistingTab = WeakSpawnedDockTab.Pin())
+	{
+		ExistingTab->UnRegisterActiveTimer(SetupWindowTimerHandle.ToSharedRef());
+	}
+
+	SetupWindowTimerHandle.Reset();
+}
+
+EActiveTimerReturnType SPopoutTabInlineContent::HandleAutoSizeWindowTimer(const double InCurrentTime
+	, const float InDeltaTime, const TWeakPtr<SDockTab> InWeakTab)
+{
+	const TSharedPtr<SDockTab> ExistingTab = InWeakTab.Pin();
+	if (!ExistingTab.IsValid())
+	{
+		SetupWindowTimerHandle.Reset();
+		return EActiveTimerReturnType::Stop;
+	}
+
+	// Wait until the window is actually constructed/visible, otherwise sizing/locking
+	// can be ignored or overwritten by later docking/window setup.
+	const TSharedPtr<SWindow> Window = ExistingTab->GetParentWindow();
+	if (!Window.IsValid() || !Window->GetNativeWindow().IsValid())
+	{
+		// Parent window not assigned yet (common on first show). Keep waiting.
+		return EActiveTimerReturnType::Continue;
+	}
+
+	const bool bIsVisible = Window->IsVisible();
+	const FVector2D SizeInScreen = Window->GetSizeInScreen();
+	const bool bHasNonZeroSize = SizeInScreen.X > 1.f && SizeInScreen.Y > 1.f;
+	if (!bIsVisible || !bHasNonZeroSize)
+	{
+		return EActiveTimerReturnType::Continue;
+	}
+
+	// An auto-sized tab only sets the size of the tab to the contents when the tab is first displayed.
+	// Switch to a fixed layout after we use the auto-sized content to get the desired size.
+	// Temporarily enable auto sizing and force a layout pass to measure the widget tree's final
+	// desired size, then use this to set the size of the window based on the content
+	Window->SetSizingRule(ESizingRule::Autosized);
+	Window->SlatePrepass(GetPrepassLayoutScaleMultiplier());
+	Window->Resize(Window->GetDesiredSize());
+	// FixedSize prevents the user from dragging window edges to resize
+	Window->SetSizingRule(ESizingRule::FixedSize);
+
+	SetupWindowTimerHandle.Reset();
+
+	return EActiveTimerReturnType::Stop;
+}
+
+void SPopoutTabInlineContent::TryHideStackTab()
+{
+	if (!bHideDockTabStackTab)
+	{
+		return;
+	}
+
+	const TSharedPtr<SDockTab> ExistingTab = WeakSpawnedDockTab.Pin();
+	if (!ExistingTab.IsValid())
+	{
+		return;
+	}
+
+	ExistingTab->SetParentDockTabStackTabWellHidden(true);
+}
+
+void SPopoutTabInlineContent::BroadcastPopoutStateChanged(const bool bInNewPoppedOut)
+{
+	if (bPoppedOut != bInNewPoppedOut)
+	{
+		bPoppedOut = bInNewPoppedOut;
+		PopoutStateChangedDelegate.ExecuteIfBound(bPoppedOut);
+	}
+}
+
+void SPopoutTabInlineContent::RestoreInlineContent()
+{
+	if (!DetachedContentWidget.IsValid())
+	{
+		return;
+	}
+
+	// Don't reparent while our spawned tab is still alive. The tab likely still owns the widget.
+	if (WeakSpawnedDockTab.IsValid())
+	{
+		return;
+	}
+
+	// Restore content back into the inline container so it has exactly one parent again.
+	// We only ever set InlineContentContainer content from this class.
+	// DetachedContentWidget being valid is our "detached and must restore" signal.
+	if (InlineContentContainer.IsValid())
+	{
+		const TSharedRef<SWidget> WidgetToRestore = DetachedContentWidget.ToSharedRef();
+		DetachedContentWidget.Reset();
+
+		InlineContentContainer->SetContent(WidgetToRestore);
+
+		return;
+	}
+
+	ensureMsgf(false, TEXT("RestoreInlineContent: InlineContentContainer invalid while DetachedContentWidget is valid."));
+
+	// Reset during teardown to avoid holding UI alive / leaking refs
+	if (!FSlateApplication::IsInitialized())
+	{
+		DetachedContentWidget.Reset();
+	}
+}
+
+void SPopoutTabInlineContent::SubscribeToTabForegrounded()
+{
+	if (IsEngineExitRequested() || !FSlateApplication::IsInitialized())
+	{
+		return;
+	}
+
+	if (TabForegroundedDelegate.IsBound())
+	{
+		return;
+	}
+
+	TabForegroundedDelegate = FOnActiveTabChanged::FDelegate::CreateSP(this, &SPopoutTabInlineContent::OnTabForegrounded);
+	FGlobalTabmanager::Get()->OnTabForegrounded_Subscribe(TabForegroundedDelegate);
+}
+
+void SPopoutTabInlineContent::UnsubscribeFromTabForegrounded()
+{
+	const FDelegateHandle Handle = TabForegroundedDelegate.GetHandle();
+	if (Handle.IsValid() && FSlateApplication::IsInitialized() && !IsEngineExitRequested())
+	{
+		FGlobalTabmanager::Get()->OnTabForegrounded_Unsubscribe(Handle);
+	}
+
+	TabForegroundedDelegate.Unbind();
+}
+
+#undef LOCTEXT_NAMESPACE
diff --git a/Engine/Source/Developer/ToolWidgets/Public/SPopoutTabInlineContent.h b/Engine/Source/Developer/ToolWidgets/Public/SPopoutTabInlineContent.h
new file mode 100644
--- /dev/null
+++ b/Engine/Source/Developer/ToolWidgets/Public/SPopoutTabInlineContent.h
@@ -0,0 +1,192 @@
+﻿// Copyright Epic Games, Inc. All Rights Reserved.
+
+#pragma once
+
+#include "Framework/SlateDelegates.h"
+#include "Styling/AppStyle.h"
+#include "Widgets/Docking/SDockTab.h"
+#include "Widgets/SCompoundWidget.h"
+#include "Widgets/SWidget.h"
+
+#define UE_API TOOLWIDGETS_API
+
+class FSpawnTabArgs;
+class FTabManager;
+class FWorkspaceItem;
+class SBox;
+class SButton;
+
+/**
+ * Inline content that can be popped out to a tab window. Manages the registration of the tab window.
+ */
+class SPopoutTabInlineContent : public SCompoundWidget
+{
+public:
+	DECLARE_DELEGATE_RetVal(bool, FOnIsPoppedOut)
+	DECLARE_DELEGATE_OneParam(FOnPopoutStateChanged, const bool /*bInPoppedOut*/);
+
+	/** Constructs a default toggle button that can be used to pop out or pop in the content */
+	UE_API static TSharedRef<SButton> CreateToggleButton(const TSharedRef<SPopoutTabInlineContent>& InOwner
+		, const FText& InTooltipText = NSLOCTEXT("ToolWidgets", "PopoutTabInlineContentToolTip", "Open this content in its own tab window")
+		, const FSlateBrush* const InIconBrush = FAppStyle::GetBrush(TEXT("Profiler.EventGraph.ExpandSelection")));
+
+	SLATE_BEGIN_ARGS(SPopoutTabInlineContent)
+		: _ContentHAlign(HAlign_Center)
+		, _ContentVAlign(VAlign_Center)
+		, _TabManager(nullptr)
+		, _TabId(NAME_None)
+		, _TabDisplayName(FText::GetEmpty())
+		, _TabIcon()
+		, _TabRole(ETabRole::NomadTab)
+		, _TabGroup(nullptr)
+		, _ShouldAutosize(true)
+		, _CanDragToPopout(true)
+		, _ShowHandle(true)
+		, _HideDockTabStackTab(false)
+		, _InitiallyPoppedOut(false)
+		, _OnPopoutStateChanged()
+		, _OnGetMenuContent()
+	{}
+		/** Inline content to pop out */
+		SLATE_DEFAULT_SLOT(FArguments, Content)
+		/** Horizontal alignment of the content */
+		SLATE_ARGUMENT(EHorizontalAlignment, ContentHAlign)
+		/** Vertical alignment of the content */
+		SLATE_ARGUMENT(EVerticalAlignment, ContentVAlign)
+
+		/** Explicit tab manager to use, otherwise uses the global tab manager */
+		SLATE_ARGUMENT(TWeakPtr<FTabManager>, TabManager)
+		/** Unique identifier for the registered tab window */
+		SLATE_ARGUMENT(FName, TabId)
+		/** Display text to show on the tab */
+		SLATE_ARGUMENT(FText, TabDisplayName)
+		/** Icon to use for the registered tab window */
+		SLATE_ARGUMENT(FSlateIcon, TabIcon)
+		/** Type of tab window to register */
+		SLATE_ARGUMENT(ETabRole, TabRole)
+		/** Tab group that determines where the menu item exists for the tab window */
+		SLATE_ARGUMENT(TSharedPtr<FWorkspaceItem>, TabGroup)
+
+		/** If true, will try to auto size the tab window to the content */
+		SLATE_ARGUMENT(bool, ShouldAutosize)
+		/** If true, can click and drag in empty space to pop out content */
+		SLATE_ARGUMENT(bool, CanDragToPopout)
+		/** If true, shows the handle image that helps let the user know they can popout content */
+		SLATE_ARGUMENT(bool, ShowHandle)
+		/** If true, hides the dock tab stack tab */
+		SLATE_ARGUMENT(bool, HideDockTabStackTab)
+		/** If true, pops out the content as soon as this widget is constructed */
+		SLATE_ARGUMENT(bool, InitiallyPoppedOut)
+
+		/** Event called when the popped out state is changed */
+		SLATE_EVENT(FOnPopoutStateChanged, OnPopoutStateChanged)
+		/** Event called to get content to be displayed for the right click context menu in empty space */
+		SLATE_EVENT(FOnGetContent, OnGetMenuContent)
+	SLATE_END_ARGS()
+
+	UE_API virtual ~SPopoutTabInlineContent() override;
+
+	UE_API void Construct(const FArguments& InArgs);
+
+	//~ Begin SWidget
+	UE_API virtual FCursorReply OnCursorQuery(const FGeometry& InGeometry, const FPointerEvent& InPointerEvent) const override;
+	UE_API virtual FReply OnMouseButtonDown(const FGeometry& InGeometry, const FPointerEvent& InPointerEvent) override;
+	UE_API virtual FReply OnMouseButtonUp(const FGeometry& InGeometry, const FPointerEvent& InPointerEvent) override;
+	UE_API virtual void OnMouseLeave(const FPointerEvent& InPointerEvent) override;
+	UE_API virtual FReply OnDragDetected(const FGeometry& InGeometry, const FPointerEvent& InPointerEvent) override;
+	//~ End SWidget
+
+	/**
+	 * Checks whether the inline content is currently popped out as a tab window.
+	 *
+	 * @return True if the content is popped out, otherwise false.
+	 */
+	UE_API bool IsPoppedOut() const;
+
+	/**
+	 * Modifies the popout state of the inline content.
+	 * Popping out the content detaches it from its current container and moves it to a new tab window,
+	 * while turning off the popout state reattaches it to its original container.
+	 *
+	 * @param bInPopout If true, the content is popped out to a separate tab window.
+	 *                  If false, the content is returned to its original container.
+	 */
+	UE_API void Popout(const bool bInPopout);
+
+	/**
+	 * Toggles the popout state of the inline content. Switches the content between its popped-out and embedded states.
+	 */
+	UE_API void TogglePopout();
+
+private:
+	/** Gets the passed in tab manager or the global tab manager none was passed or null */
+	TSharedPtr<FTabManager> GetTabManager() const;
+
+	/** Checks if the passed in tab Id is unique and does not conflict with any other existing tabs */
+	bool IsUniqueTabId() const;
+
+	void RegisterTab();
+	void UnregisterTab();
+
+	TSharedRef<SDockTab> SpawnTab(const FSpawnTabArgs& InArgs);
+	void CloseTab();
+
+	void OnTabActivated(const TSharedRef<SDockTab> InDockTab, const ETabActivationCause InActivationCause);
+	void OnTabClosed(const TSharedRef<SDockTab> InDockTab);
+	void OnTabRelocated();
+	void OnTabForegrounded(const TSharedPtr<SDockTab> InNewTab, const TSharedPtr<SDockTab> InPreviousTab);
+
+	TSharedRef<SWidget> CreateDefaultMenuContent();
+
+	EVisibility GetContentVisibility() const;
+	EVisibility GetContentBorderVisibility() const;
+
+	void RegisterAutoSizeTimer(const TSharedPtr<SDockTab>& InDockTab = nullptr);
+	void UnregisterAutoSizeTimer();
+	EActiveTimerReturnType HandleAutoSizeWindowTimer(const double InCurrentTime
+		, const float InDeltaTime, const TWeakPtr<SDockTab> InWeakTab);
+
+	void TryHideStackTab();
+
+	void BroadcastPopoutStateChanged(const bool bInNewPoppedOut);
+
+	void RestoreInlineContent();
+
+	void SubscribeToTabForegrounded();
+	void UnsubscribeFromTabForegrounded();
+
+	/** Cached slate constructor arguments */
+	TWeakPtr<SWidget> WeakContentWidget;
+	TWeakPtr<FTabManager> WeakTabManager;
+	FName TabId;
+	FText TabDisplayName;
+	FSlateIcon TabIcon;
+	ETabRole TabRole = ETabRole::NomadTab;
+	TWeakPtr<FWorkspaceItem> WeakTabGroup;
+	bool bShouldAutosize = true;
+	bool bCanDragToPopout = true;
+	bool bShowHandle = true;
+	bool bHideDockTabStackTab = false;
+	FOnPopoutStateChanged PopoutStateChangedDelegate;
+	FOnGetContent GetMenuContentDelegate;
+
+	TWeakPtr<SDockTab> WeakSpawnedDockTab;
+
+	bool bTabRegistered = false;
+	bool bPoppedOut = false;
+	bool bIsForeground = false;
+
+	bool bLeftMouseDownOverWidget = false;
+	bool bRightMouseDownOverWidget = false;
+
+	FOnActiveTabChanged::FDelegate TabForegroundedDelegate;
+
+	TSharedPtr<FActiveTimerHandle> SetupWindowTimerHandle;
+
+	/** Inline container that owns the content when not popped out */
+	TSharedPtr<SBox> InlineContentContainer;
+	/** Strong ref while content is moved out of the inline host and into the dock tab */
+	TSharedPtr<SWidget> DetachedContentWidget;
+};
+
+#undef UE_API
