Fix header path in deprecation message, making desired user action more obvious

#rnx

#Synced-CL 50958978

diff --git a/Engine/Source/Developer/ToolWidgets/Public/SSearchableComboBox.h b/Engine/Source/Developer/ToolWidgets/Public/SSearchableComboBox.h
--- a/Engine/Source/Developer/ToolWidgets/Public/SSearchableComboBox.h
+++ b/Engine/Source/Developer/ToolWidgets/Public/SSearchableComboBox.h
@@ -4,7 +4,7 @@
 
 #include "Misc/CoreMiscDefines.h"
 
-UE_DEPRECATED_HEADER(5.8, "This header has moved to the Slate module. Slate/Widgets/Input/SSearchableComboBox.h.")
+UE_DEPRECATED_HEADER(5.8, "This header has moved to the Slate module - Widgets/Input/SSearchableComboBox.h.")
 
 // Include the new file so that the project still compiles, but has warnings
 #include "Widgets/Input/SSearchableComboBox.h"
