PR #14238: Enable delayed change notification on SceneOutliner search

#Synced-CL 51758326

diff --git a/Engine/Source/Developer/ToolWidgets/Private/Filters/SFilterSearchBox.cpp b/Engine/Source/Developer/ToolWidgets/Private/Filters/SFilterSearchBox.cpp
--- a/Engine/Source/Developer/ToolWidgets/Private/Filters/SFilterSearchBox.cpp
+++ b/Engine/Source/Developer/ToolWidgets/Private/Filters/SFilterSearchBox.cpp
@@ -27,6 +27,7 @@ public:
 		, _OnTextChanged()
 		, _OnTextCommitted()
 		, _DelayChangeNotificationsWhileTyping( true )
+		, _DelayChangeNotificationsWhileTypingSeconds( 0.25f )
 	{ }
 
 		/** The text displayed in the SearchBox when no text has been entered */
@@ -44,6 +45,9 @@ public:
 		/** Whether the SearchBox should delay notifying listeners of text changed events until the user is done typing */
 		SLATE_ATTRIBUTE( bool, DelayChangeNotificationsWhileTyping )
 
+		/** If we're delaying change notifications, how many seconds should we wait? */
+		SLATE_ATTRIBUTE( float, DelayChangeNotificationsWhileTypingSeconds )
+
 		/** Callback delegate to have first chance handling of the OnKeyDown event */
 		SLATE_EVENT(FOnKeyDown, OnKeyDownHandler)
 
@@ -82,6 +86,7 @@ void SFilterSearchBoxImpl::Construct(const FArguments& InArgs)
 		.OnTextCommitted(InArgs._OnTextCommitted)
 		.SelectAllTextWhenFocused( false )
 		.DelayChangeNotificationsWhileTyping( InArgs._DelayChangeNotificationsWhileTyping )
+		.DelayChangeNotificationsWhileTypingSeconds( InArgs._DelayChangeNotificationsWhileTypingSeconds )
 		.OnKeyDownHandler(InArgs._OnKeyDownHandler)
 	);
 
@@ -184,6 +189,7 @@ void SFilterSearchBox::Construct( const FArguments& InArgs )
 			.OnTextChanged(this, &SFilterSearchBox::HandleTextChanged)
 			.OnTextCommitted(this, &SFilterSearchBox::HandleTextCommitted)
 			.DelayChangeNotificationsWhileTyping( InArgs._DelayChangeNotificationsWhileTyping )
+			.DelayChangeNotificationsWhileTypingSeconds( InArgs._DelayChangeNotificationsWhileTypingSeconds )
 			.OnKeyDownHandler(InArgs._OnKeyDownHandler)
 			.OnSearchHistoryClicked(this, &SFilterSearchBox::OnClickedSearchHistory)
 			.GetSearchHistoryVisibility(this, &SFilterSearchBox::GetSearchHistoryVisibility)
diff --git a/Engine/Source/Developer/ToolWidgets/Public/Filters/SFilterSearchBox.h b/Engine/Source/Developer/ToolWidgets/Public/Filters/SFilterSearchBox.h
--- a/Engine/Source/Developer/ToolWidgets/Public/Filters/SFilterSearchBox.h
+++ b/Engine/Source/Developer/ToolWidgets/Public/Filters/SFilterSearchBox.h
@@ -33,6 +33,7 @@ public:
 	, _ShowSearchHistory(true)
 	, _MaxSearchHistory(5)
 	, _DelayChangeNotificationsWhileTyping( false )
+	, _DelayChangeNotificationsWhileTypingSeconds(0.25f)
 	{}
 
 	/** Invoked whenever the text changes */
@@ -58,6 +59,9 @@ public:
 
 	/** Whether the SearchBox should delay notifying listeners of text changed events until the user is done typing */
 	SLATE_ATTRIBUTE( bool, DelayChangeNotificationsWhileTyping )
+	
+	/** If we're delaying change notifications, how many seconds should we wait? */
+    SLATE_ATTRIBUTE( float, DelayChangeNotificationsWhileTypingSeconds )
 
 	/** Callback delegate to have first chance handling of the OnKeyDown event */
 	SLATE_EVENT(FOnKeyDown, OnKeyDownHandler)
diff --git a/Engine/Source/Editor/SceneOutliner/Private/SSceneOutliner.cpp b/Engine/Source/Editor/SceneOutliner/Private/SSceneOutliner.cpp
--- a/Engine/Source/Editor/SceneOutliner/Private/SSceneOutliner.cpp
+++ b/Engine/Source/Editor/SceneOutliner/Private/SSceneOutliner.cpp
@@ -66,7 +66,9 @@ void SSceneOutliner::Construct(const FArguments& InArgs, const FSceneOutlinerIni
 		.Visibility( InInitOptions.bShowSearchBox ? EVisibility::Visible : EVisibility::Collapsed )
 		.HintText( LOCTEXT( "FilterSearch", "Search..." ) )
 		.OnTextChanged( this, &SSceneOutliner::OnFilterTextChanged )
-		.OnTextCommitted( this, &SSceneOutliner::OnFilterTextCommitted );
+		.OnTextCommitted( this, &SSceneOutliner::OnFilterTextCommitted )
+		.DelayChangeNotificationsWhileTyping( true )
+		.DelayChangeNotificationsWhileTypingSeconds( 0.2f );
 	}
 	
 	CreateFilterBar(InInitOptions.FilterBarOptions);
