SSegmentedProgressBar updates:
Moved to ToolWidgets with style decoupling. This allows SSegmentedProgressBar to be used by any Slate code that needs it.
Updated the segments to offer an optional label widget for further customization.
Added an optional Skipped state

[RN] SSegmentedProgressBar is now available for displaying the progress of a multi-state workflow.
#Synced-CL 52440464

diff --git a/Engine/Plugins/Developer/ProjectLauncher/Source/ProjectLauncher/Private/Styling/ProjectLauncherStyle.cpp b/Engine/Plugins/Developer/ProjectLauncher/Source/ProjectLauncher/Private/Styling/ProjectLauncherStyle.cpp
--- a/Engine/Plugins/Developer/ProjectLauncher/Source/ProjectLauncher/Private/Styling/ProjectLauncherStyle.cpp
+++ b/Engine/Plugins/Developer/ProjectLauncher/Source/ProjectLauncher/Private/Styling/ProjectLauncherStyle.cpp
@@ -19,6 +19,7 @@
 #define RootToCoreContentDir StyleSet->RootToCoreContentDir
 
 TSharedPtr<FSlateStyleSet> FProjectLauncherStyle::StyleSet = nullptr;
+FSegmentedProgressBarStyle FProjectLauncherStyle::ProgressBarStyle;
 
 static const FVector2D Icon12x12(12.0f, 12.0f);
 static const FVector2D Icon15x15(15.0f, 15.0f);
@@ -56,13 +57,6 @@ void FProjectLauncherStyle::Initialize()
 	StyleSet->Set("OuterCircle.Busy", new IMAGE_BRUSH(TEXT("outer_circle_busy_512px"), Icon512x512));
 	StyleSet->Set("FullCircle", new IMAGE_BRUSH_SVG(TEXT("full_circle_white"), Icon36x36));
 
-	StyleSet->Set("State.Busy", FLinearColor(FColor(0xFF0070E0)) );
-	StyleSet->Set("State.Pending", FLinearColor(FColor(0xFF1A1A1A)) );
-	StyleSet->Set("State.Error", FLinearColor(FColor(0xFFEF3535)) );
-	StyleSet->Set("State.Complete", FLinearColor(FColor(0xFF0070E0)) );
-	StyleSet->Set("State.Canceled", FLinearColor(FColor(0xFF383838)) );
-	StyleSet->Set("State.AllComplete", FLinearColor(FColor(0xFF0070E0)) );
-
 	FButtonStyle HoverHintOnly = FButtonStyle()
 		.SetNormal(FSlateNoResource())
 		.SetHovered(CORE_BOX_BRUSH("../Editor/Slate/Common/ButtonHoverHint", FMargin(4.f/16.0f), FLinearColor(1.f,1.f,1.f,0.15f)))
@@ -71,6 +65,20 @@ void FProjectLauncherStyle::Initialize()
 		.SetPressedPadding(FMargin(0.f,1.f,0.f,0.f));
 	StyleSet->Set( "HoverHintOnly", HoverHintOnly );
 
+	ProgressBarStyle = FSegmentedProgressBarStyle()
+		.SetBusyColor(FLinearColor(FColor(0xFF0070E0)))
+		.SetPendingColor(FLinearColor(FColor(0xFF1A1A1A)))
+		.SetErrorColor(FLinearColor(FColor(0xFFEF3535)))
+		.SetCompleteColor(FLinearColor(FColor(0xFF0070E0)))
+		.SetCanceledColor(FLinearColor(FColor(0xFF383838)))
+		.SetAllCompleteColor(FLinearColor(FColor(0xFF0070E0)))
+		.SetFullCircleBrush(StyleSet->GetBrush("FullCircle"))
+		.SetOuterCircleBrush(StyleSet->GetBrush("OuterCircle"))
+		.SetOuterCircleBusyBrush(StyleSet->GetBrush("OuterCircle.Busy"))
+		.SetBadgeSuccessBrush(StyleSet->GetBrush("BadgeOutlined.Success"))
+		.SetBadgeErrorBrush(StyleSet->GetBrush("BadgeOutlined.Error"))
+		.SetBadgeAllCompleteBrush(StyleSet->GetBrush("BadgeOutlined.AllComplete"));
+
 	StyleSet->Set("Icons.DiffersFromDefault", new CORE_IMAGE_BRUSH_SVG( "../Editor/Slate/Starship/Common/ResetToDefault", Icon16x16) );
 	StyleSet->Set("Icons.DeviceManager", new CORE_IMAGE_BRUSH_SVG( "../Editor/Slate/Starship/Common/DeviceManager", Icon16x16) );
 	StyleSet->Set("Icons.Plugins", new CORE_IMAGE_BRUSH_SVG("../Editor/Slate/Starship/Common/Plugins", Icon16x16));
diff --git a/Engine/Plugins/Developer/ProjectLauncher/Source/ProjectLauncher/Private/Styling/ProjectLauncherStyle.h b/Engine/Plugins/Developer/ProjectLauncher/Source/ProjectLauncher/Private/Styling/ProjectLauncherStyle.h
--- a/Engine/Plugins/Developer/ProjectLauncher/Source/ProjectLauncher/Private/Styling/ProjectLauncherStyle.h
+++ b/Engine/Plugins/Developer/ProjectLauncher/Source/ProjectLauncher/Private/Styling/ProjectLauncherStyle.h
@@ -5,6 +5,7 @@
 #include "Styling/SlateStyle.h"
 #include "ILauncherTask.h"
 #include "Misc/DataDrivenPlatformInfoRegistry.h"
+#include "SSegmentedProgressBar.h"
 
 namespace PlatformInfo
 {
@@ -28,6 +29,12 @@ public:
 	static const FSlateBrush* GetProfileBrushForPlatform(const PlatformInfo::FTargetPlatformInfo* PlatformInfo, EPlatformIconSize IconSize);
 	static const FSlateBrush* GetProfileBrushForPlatforms(TArray<const PlatformInfo::FTargetPlatformInfo*> PlatformInfos, EPlatformIconSize IconSize);
 
+	static const FSegmentedProgressBarStyle& GetProgressBarStyle()
+	{
+		return ProgressBarStyle;
+	}
+
 private:
 	static TSharedPtr<FSlateStyleSet> StyleSet;	
+	static FSegmentedProgressBarStyle ProgressBarStyle;
 };
diff --git a/Engine/Plugins/Developer/ProjectLauncher/Source/ProjectLauncher/Private/Widgets/Output/SCustomLaunchLaunchPanel.cpp b/Engine/Plugins/Developer/ProjectLauncher/Source/ProjectLauncher/Private/Widgets/Output/SCustomLaunchLaunchPanel.cpp
--- a/Engine/Plugins/Developer/ProjectLauncher/Source/ProjectLauncher/Private/Widgets/Output/SCustomLaunchLaunchPanel.cpp
+++ b/Engine/Plugins/Developer/ProjectLauncher/Source/ProjectLauncher/Private/Widgets/Output/SCustomLaunchLaunchPanel.cpp
@@ -160,6 +160,7 @@ void SCustomLaunchLaunchPanel::Construct(const FArguments& InArgs, const TShared
 				.AutoHeight()
 				[
 					SAssignNew(SubwayProgressBar, SSegmentedProgressBar)
+					.Style(FProjectLauncherStyle::GetProgressBarStyle())
 				]
 #endif
 
diff --git a/Engine/Plugins/Developer/ProjectLauncher/Source/ProjectLauncher/Private/Widgets/Output/SCustomLaunchLaunchPanel.h b/Engine/Plugins/Developer/ProjectLauncher/Source/ProjectLauncher/Private/Widgets/Output/SCustomLaunchLaunchPanel.h
--- a/Engine/Plugins/Developer/ProjectLauncher/Source/ProjectLauncher/Private/Widgets/Output/SCustomLaunchLaunchPanel.h
+++ b/Engine/Plugins/Developer/ProjectLauncher/Source/ProjectLauncher/Private/Widgets/Output/SCustomLaunchLaunchPanel.h
@@ -5,7 +5,7 @@
 #include "Internationalization/Text.h"
 #include "Framework/SlateDelegates.h"
 #include "Widgets/SCompoundWidget.h"
-#include "Widgets/SSegmentedProgressBar.h"
+#include "SSegmentedProgressBar.h"
 #include "Widgets/Output/CustomLaunchOutputLogMarshaller.h"
 #include "Model/ProjectLauncherModel.h"
 #include "ILauncherWorker.h"
diff --git a/Engine/Plugins/Developer/ProjectLauncher/Source/ProjectLauncher/Private/Widgets/SSegmentedProgressBar.cpp b/Engine/Plugins/Developer/ProjectLauncher/Source/ProjectLauncher/Private/Widgets/SSegmentedProgressBar.cpp
deleted file mode 100644
diff --git a/Engine/Plugins/Developer/ProjectLauncher/Source/ProjectLauncher/Private/Widgets/SSegmentedProgressBar.h b/Engine/Plugins/Developer/ProjectLauncher/Source/ProjectLauncher/Private/Widgets/SSegmentedProgressBar.h
deleted file mode 100644
diff --git a/Engine/Source/Developer/ToolWidgets/Private/SSegmentedProgressBar.cpp b/Engine/Source/Developer/ToolWidgets/Private/SSegmentedProgressBar.cpp
new file mode 100644
--- /dev/null
+++ b/Engine/Source/Developer/ToolWidgets/Private/SSegmentedProgressBar.cpp
@@ -0,0 +1,431 @@
+// Copyright Epic Games, Inc. All Rights Reserved.
+
+#include "SSegmentedProgressBar.h"
+#include "Widgets/Images/SImage.h"
+#include "Widgets/Layout/SBox.h"
+#include "Widgets/SBoxPanel.h"
+#include "Widgets/SOverlay.h"
+#include "Widgets/Colors/SColorBlock.h"
+#include "Styling/StyleColors.h"
+#include "Layout/LayoutUtils.h"
+
+#define LOCTEXT_NAMESPACE "SSegmentedProgressBar"
+
+SSegmentedProgressBar::FScopedWidgetSlotArguments SSegmentedProgressBar::InsertSlot(int32 Index, bool bRebuildChildren)
+{
+	if (bRebuildChildren)
+	{
+		TWeakPtr<SSegmentedProgressBar> AsWeak = SharedThis(this);
+		return FScopedWidgetSlotArguments { MakeUnique<FSlot>(), this->Children, Index, [AsWeak](const FSlot*, int32)
+		{
+			if (TSharedPtr<SSegmentedProgressBar> SharedThis = AsWeak.Pin())
+			{
+				SharedThis->RebuildChildren();
+			}
+		}};
+	}
+	else
+	{
+		return FScopedWidgetSlotArguments(MakeUnique<FSlot>(), this->Children, Index);
+	}
+}
+
+
+
+void SSegmentedProgressBar::Construct( const FArguments& InArgs )
+{
+	Style = InArgs._Style;
+	ThrobberAnimation = FCurveSequence(0.0f, 1.0f);
+	ThrobberAnimation.Play(AsShared(), true);
+
+	Children.AddSlots(MoveTemp(const_cast<TArray<typename FSlot::FSlotArguments>&>(InArgs._Slots)));	
+	RebuildChildren();
+}
+
+
+
+void SSegmentedProgressBar::Tick(const FGeometry& AllottedGeometry, const double InCurrentTime, const float InDeltaTime)
+{
+	SCompoundWidget::Tick(AllottedGeometry, InCurrentTime, InDeltaTime);
+	CacheChildStates();
+}
+
+
+
+void SSegmentedProgressBar::CacheChildStates()
+{
+	CachedChildStates.Reset();
+	for (int32 ChildIndex = 0; ChildIndex < Children.Num(); ChildIndex++)
+	{
+		EState ChildState = Children[ChildIndex].State.Get();
+
+		// show us as busy as soon as the previous task is completed - just looks nicer
+		if (ChildState == EState::Pending && (ChildIndex == 0 || CachedChildStates[ChildIndex-1] == EState::Completed))
+		{
+			ChildState = EState::Busy;
+		}
+
+		CachedChildStates.Add( ChildState );
+	}
+
+	if (CachedChildStates.Num() == 0 )
+	{
+		CachedOverallState = EState::None;
+	}
+	else
+	{
+		CachedOverallState = CachedChildStates.Last();
+	}
+}
+
+
+
+SSegmentedProgressBar::EState SSegmentedProgressBar::GetChildState( int32 ChildIndex ) const
+{
+	if (CachedChildStates.IsValidIndex(ChildIndex))
+	{
+		return CachedChildStates[ChildIndex];
+	}
+	
+	return EState::None;
+}
+
+
+
+void SSegmentedProgressBar::RebuildChildren()
+{
+	CacheChildStates();
+
+	TSharedPtr<SHorizontalBox> CircleRow;
+	
+	TSharedPtr<SVerticalBox> Container;
+	ChildSlot
+	[ 
+		SAssignNew(Container, SVerticalBox)
+		+ SVerticalBox::Slot()
+		.AutoHeight()
+		[
+			SAssignNew(CircleRow, SHorizontalBox)
+		]
+	];
+
+	// Should we add optional labels?
+	TSharedPtr<SHorizontalBox> LabelRow;
+
+	bool bHasLabels = false;
+	for (int32 i = 0; i < Children.Num(); ++i)
+	{
+		if (Children[i].LabelContent.IsValid())
+		{
+			bHasLabels = true;
+			break;
+		}
+	}
+	if (bHasLabels)
+	{
+		Container->AddSlot()
+		.AutoHeight()
+		.Padding(0.0f, 4.0f, 0.0f, 0.0f)
+		[
+			SAssignNew(LabelRow, SHorizontalBox)
+		];
+	}
+
+	const int32 NumChildren = Children.Num();
+	for ( int32 ChildIndex = 0; ChildIndex < NumChildren; ChildIndex++ )
+	{
+		TSharedRef<SWidget> Child = Children[ChildIndex].GetWidget();
+		FSlot* ChildSlotPtr = &Children[ChildIndex];
+
+		if (Child == SNullWidget::NullWidget)
+		{
+			Child = ConstructChild(*ChildSlotPtr, ChildIndex);
+		}
+
+		// separator bar between previous and current
+		if (ChildIndex > 0)
+		{
+			AddChildSeparatorBar(CircleRow, ChildIndex);
+			if (LabelRow)
+			{
+				LabelRow->AddSlot()
+				.FillWidth(1)
+				[ 
+					SNullWidget::NullWidget 
+				];
+			}
+		}
+
+		// task item
+		CircleRow->AddSlot()
+		.VAlign(VAlign_Top)
+		.HAlign(HAlign_Center)
+		.AutoWidth()
+		[
+			Child
+		];
+
+		if (LabelRow)
+		{
+			LabelRow->AddSlot()
+			.AutoWidth()
+			[
+				SNew(SBox)
+				.WidthOverride(36.0f)
+				[
+					SNew(SOverlay)
+					+ SOverlay::Slot()
+					.HAlign(HAlign_Center)
+					.VAlign(VAlign_Fill)
+					.Padding(-1000.0f, 0.0f)
+					[
+						ChildSlotPtr->LabelContent.IsValid()
+							? ChildSlotPtr->LabelContent.ToSharedRef()
+							: SNullWidget::NullWidget
+					]
+				]
+			];
+		}
+	}
+}
+
+
+
+void SSegmentedProgressBar::AddChildSeparatorBar( TSharedPtr<SHorizontalBox> SlotBox, int32 ChildIndex ) const
+{
+	SlotBox->AddSlot()
+	.VAlign(VAlign_Top)
+	.HAlign(HAlign_Fill)
+	.FillWidth(1)
+	.Padding(0.0f, (36.0f - LineSize) / 2.0f, 0.0f, 0.0f)
+	[
+		SNew(SColorBlock)
+		.Color( this, &SSegmentedProgressBar::GetSeparatorBarColor, ChildIndex )
+		.Size( FVector2D(32, LineSize) )
+	];
+}
+
+
+
+TSharedRef<SWidget> SSegmentedProgressBar::ConstructChild( const FSlot& Slot, int32 ChildIndex ) const
+{
+	TSharedRef<SOverlay> CircleOverlay = SNew(SOverlay)
+	.ToolTipText(Slot.ToolTipText)
+
+	// full circle (only shown when fully complete or cancelled)
+	+SOverlay::Slot()
+	.VAlign(VAlign_Center)
+	.HAlign(HAlign_Center)
+	[
+		SNew(SImage)
+		.Image( Style.FullCircleBrush )
+		.ColorAndOpacity( this, &SSegmentedProgressBar::GetCircleColor, ChildIndex)
+		.Visibility( this, &SSegmentedProgressBar::GetFullCircleVisibility, ChildIndex )
+
+	]
+
+	// outer circle
+	+SOverlay::Slot()
+	.VAlign(VAlign_Center)
+	.HAlign(HAlign_Center)
+	[
+		SNew(SImage)
+		.DesiredSizeOverride(FVector2D(36,36))
+		.Image( Style.OuterCircleBrush )
+		.ColorAndOpacity( this, &SSegmentedProgressBar::GetCircleColor, ChildIndex )
+		.Visibility( this, &SSegmentedProgressBar::GetOuterCircleVisibility, ChildIndex )
+	]
+
+	// outer busy circle
+	+SOverlay::Slot()
+	.VAlign(VAlign_Center)
+	.HAlign(HAlign_Center)
+	[
+		SNew(SImage)
+		.DesiredSizeOverride(FVector2D(36,36))
+		.Image( Style.OuterCircleBusyBrush )
+		.ColorAndOpacity( Style.BusyColor )
+		.Visibility( this, &SSegmentedProgressBar::GetProgressCircleVisibility, ChildIndex )
+		.RenderTransform( this, &SSegmentedProgressBar::GetProgressCircleTransform, ChildIndex)
+		.RenderTransformPivot(FVector2D(.5f,.5f))
+	]
+
+	// task icon
+	+SOverlay::Slot()
+	.VAlign(VAlign_Center)
+	.HAlign(HAlign_Center)
+	[
+		SNew(SImage)
+		.Image(Slot.Image)
+		.ColorAndOpacity(this, &SSegmentedProgressBar::GetIconColor, ChildIndex)
+		.DesiredSizeOverride(FVector2D(16,16))
+	]
+
+	// task state overlay
+	+SOverlay::Slot()
+	.VAlign(VAlign_Bottom)
+	.HAlign(HAlign_Right)
+	[
+		SNew(SImage)
+		.Image(this, &SSegmentedProgressBar::GetOverlayIcon, ChildIndex)
+		.Visibility( this, &SSegmentedProgressBar::GetOverlayVisibility, ChildIndex)
+	]
+	;
+
+	return SNew(SVerticalBox)
+	+ SVerticalBox::Slot()
+	.AutoHeight()
+	.HAlign(HAlign_Center)
+	[
+		CircleOverlay
+	];
+}
+
+
+FLinearColor SSegmentedProgressBar::GetSeparatorBarColor( int32 ChildIndex ) const
+{
+	if (CachedOverallState == EState::Completed)
+	{
+		return Style.AllCompleteColor.GetSpecifiedColor();
+	}
+
+	EState State = GetChildState(ChildIndex);
+	switch (State)
+	{
+		case EState::Busy:		return Style.BusyColor.GetSpecifiedColor();
+		case EState::Canceled:  return Style.CanceledColor.GetSpecifiedColor();
+		case EState::Skipped:   return Style.CanceledColor.GetSpecifiedColor(); // skipped can use canceled's color
+		case EState::Completed: return Style.CompleteColor.GetSpecifiedColor(); // line indicates previous step succeeded in this case
+		case EState::Failed:	return Style.CompleteColor.GetSpecifiedColor(); // failed can use complete's color
+		case EState::Pending:	return Style.PendingColor.GetSpecifiedColor();
+	}
+
+	return FSlateColor::UseForeground().GetSpecifiedColor();
+}
+
+FSlateColor SSegmentedProgressBar::GetCircleColor( int32 ChildIndex ) const
+{
+	if (CachedOverallState == EState::Completed)
+	{
+		return Style.AllCompleteColor;
+	}
+
+	EState State = GetChildState(ChildIndex);
+	switch (State)
+	{
+		case EState::Busy:		return Style.PendingColor;
+		case EState::Canceled:  return Style.CanceledColor;
+		case EState::Skipped:   return Style.CanceledColor; // skipped can use canceled's color
+		case EState::Completed:	return Style.CompleteColor;
+		case EState::Failed:	return Style.ErrorColor;
+		case EState::Pending:	return Style.PendingColor;
+	}
+
+	return FSlateColor::UseForeground();
+}
+
+
+
+
+
+FSlateColor SSegmentedProgressBar::GetIconColor( int32 ChildIndex ) const
+{
+	if (CachedOverallState == EState::Completed)
+	{
+		return FStyleColors::White;
+	}
+
+	EState State = GetChildState(ChildIndex);
+	if (State == EState::Canceled || State == EState::Skipped)
+	{
+		return FStyleColors::Hover2;
+	}
+
+	return FStyleColors::Foreground;
+}
+
+
+EVisibility SSegmentedProgressBar::GetFullCircleVisibility( int32 ChildIndex ) const
+{
+	EState ChildState = GetChildState(ChildIndex);
+	if (CachedOverallState == EState::Completed || ChildState == EState::Canceled || ChildState == EState::Skipped)
+	{
+		return EVisibility::Visible;
+	}
+
+	return EVisibility::Collapsed;
+}
+
+
+
+TOptional<FSlateRenderTransform> SSegmentedProgressBar::GetProgressCircleTransform( int32 ChildIndex ) const
+{
+	if (GetChildState(ChildIndex) == EState::Busy)
+	{
+		const float DeltaAngle = ThrobberAnimation.GetLerp()*2*PI;
+		return FSlateRenderTransform(FQuat2D(DeltaAngle));
+	}
+
+	FSlateRenderTransform Result;
+	return Result;
+}
+
+
+
+EVisibility SSegmentedProgressBar::GetOuterCircleVisibility( int32 ChildIndex ) const
+{
+	EState ChildState = GetChildState(ChildIndex);
+	if (CachedOverallState != EState::Completed && ChildState != EState::Canceled && ChildState != EState::Skipped)
+	{
+		return EVisibility::Visible;
+	}
+
+	return EVisibility::Collapsed;
+}
+
+
+
+EVisibility SSegmentedProgressBar::GetProgressCircleVisibility( int32 ChildIndex ) const
+{
+	if (CachedOverallState != EState::Completed && GetChildState(ChildIndex) == EState::Busy)
+	{
+		return EVisibility::Visible;
+	}
+
+	return EVisibility::Collapsed;
+}
+
+
+
+
+EVisibility SSegmentedProgressBar::GetOverlayVisibility( int32 ChildIndex ) const
+{
+	EState State = GetChildState(ChildIndex);
+	if (State == EState::Completed || State == EState::Failed)
+	{
+		return EVisibility::Visible;
+	}
+
+	return EVisibility::Collapsed;
+}
+
+
+
+const FSlateBrush* SSegmentedProgressBar::GetOverlayIcon( int32 ChildIndex ) const
+{
+	if (CachedOverallState == EState::Completed)
+	{
+		return Style.BadgeAllCompleteBrush;
+	}
+
+	EState State = GetChildState(ChildIndex);
+	switch (State)
+	{
+		case EState::Completed: return Style.BadgeSuccessBrush;
+		case EState::Failed:	return Style.BadgeErrorBrush;
+	}
+
+	return FStyleDefaults::GetNoBrush();
+}
+
+#undef LOCTEXT_NAMESPACE
\ No newline at end of file
diff --git a/Engine/Source/Developer/ToolWidgets/Public/SSegmentedProgressBar.h b/Engine/Source/Developer/ToolWidgets/Public/SSegmentedProgressBar.h
new file mode 100644
--- /dev/null
+++ b/Engine/Source/Developer/ToolWidgets/Public/SSegmentedProgressBar.h
@@ -0,0 +1,180 @@
+// Copyright Epic Games, Inc. All Rights Reserved.
+
+#pragma once
+
+#include "CoreMinimal.h"
+#include "Styling/SlateColor.h"
+#include "Styling/CoreStyle.h"
+#include "Widgets/SCompoundWidget.h"
+#include "Animation/CurveSequence.h"
+
+#define UE_API TOOLWIDGETS_API
+
+struct FSegmentedProgressBarStyle
+{
+	FSlateColor BusyColor;
+	FSlateColor PendingColor;
+	FSlateColor ErrorColor;
+	FSlateColor CompleteColor;
+	FSlateColor CanceledColor;
+	FSlateColor AllCompleteColor;
+
+	const FSlateBrush* FullCircleBrush = nullptr;
+	const FSlateBrush* OuterCircleBrush = nullptr;
+	const FSlateBrush* OuterCircleBusyBrush = nullptr;
+	const FSlateBrush* BadgeSuccessBrush = nullptr;
+	const FSlateBrush* BadgeErrorBrush = nullptr;
+	const FSlateBrush* BadgeAllCompleteBrush = nullptr;
+
+	FSegmentedProgressBarStyle& SetBusyColor(FSlateColor InColor) { BusyColor = InColor; return *this; }
+	FSegmentedProgressBarStyle& SetPendingColor(FSlateColor InColor) { PendingColor = InColor; return *this; }
+	FSegmentedProgressBarStyle& SetErrorColor(FSlateColor InColor) { ErrorColor = InColor; return *this; }
+	FSegmentedProgressBarStyle& SetCompleteColor(FSlateColor InColor) { CompleteColor = InColor; return *this; }
+	FSegmentedProgressBarStyle& SetCanceledColor(FSlateColor InColor) { CanceledColor = InColor; return *this; }
+	FSegmentedProgressBarStyle& SetAllCompleteColor(FSlateColor InColor) { AllCompleteColor = InColor; return *this; }
+
+	FSegmentedProgressBarStyle& SetFullCircleBrush(const FSlateBrush* InBrush) { FullCircleBrush = InBrush; return *this; }
+	FSegmentedProgressBarStyle& SetOuterCircleBrush(const FSlateBrush* InBrush) { OuterCircleBrush = InBrush; return *this; }
+	FSegmentedProgressBarStyle& SetOuterCircleBusyBrush(const FSlateBrush* InBrush) { OuterCircleBusyBrush = InBrush; return *this; }
+	FSegmentedProgressBarStyle& SetBadgeSuccessBrush(const FSlateBrush* InBrush) { BadgeSuccessBrush = InBrush; return *this; }
+	FSegmentedProgressBarStyle& SetBadgeErrorBrush(const FSlateBrush* InBrush) { BadgeErrorBrush = InBrush; return *this; }
+	FSegmentedProgressBarStyle& SetBadgeAllCompleteBrush(const FSlateBrush* InBrush) { BadgeAllCompleteBrush = InBrush; return *this; }
+};
+
+class SSegmentedProgressBar : public SCompoundWidget
+{
+public:
+	enum class EState : uint8
+	{
+		None,
+		Busy,
+		Canceled,
+		Completed,
+		Failed,
+		Pending,
+		Skipped,
+	};
+
+
+	struct FSlot : public TSlotBase<FSlot>
+	{
+	public:
+		FSlot()
+			: TSlotBase<FSlot>()
+			, Image(nullptr)
+		{}
+
+		SLATE_SLOT_BEGIN_ARGS(FSlot, TSlotBase<FSlot>)
+			SLATE_ATTRIBUTE(const FSlateBrush*, Image)
+			SLATE_ATTRIBUTE(EState, State)
+			SLATE_ATTRIBUTE(FText, ToolTipText)
+			SLATE_ARGUMENT(TSharedPtr<SWidget>, LabelContent)
+		SLATE_SLOT_END_ARGS()
+
+		void Construct(const FChildren& SlotOwner, FSlotArguments&& InArgs)
+		{
+			TSlotBase<FSlot>::Construct(SlotOwner, MoveTemp(InArgs));
+			if (InArgs._Image.IsSet())
+			{
+				Image = MoveTemp(InArgs._Image);
+			}
+			if (InArgs._State.IsSet())
+			{
+				State = MoveTemp(InArgs._State);
+			}
+			if (InArgs._ToolTipText.IsSet())
+			{
+				ToolTipText = MoveTemp(InArgs._ToolTipText);
+			}
+			if (InArgs._LabelContent)
+			{
+				LabelContent = InArgs._LabelContent;
+			}
+		}
+
+	protected:
+		TAttribute<const FSlateBrush*> Image;
+		TAttribute<EState> State;
+		TAttribute<FText> ToolTipText;
+		TSharedPtr<SWidget> LabelContent;
+
+		friend SSegmentedProgressBar;
+	};
+
+	static FSlot::FSlotArguments Slot()
+	{
+		return FSlot::FSlotArguments(MakeUnique<FSlot>());
+	}
+
+	SLATE_BEGIN_ARGS(SSegmentedProgressBar){}
+		SLATE_SLOT_ARGUMENT(FSlot, Slots)
+		SLATE_ARGUMENT(FSegmentedProgressBarStyle, Style)
+	SLATE_END_ARGS()
+
+	using FScopedWidgetSlotArguments = typename TPanelChildren<FSlot>::FScopedWidgetSlotArguments;
+	FScopedWidgetSlotArguments AddSlot(bool bRebuildChildren = true)
+	{
+		return InsertSlot(INDEX_NONE, bRebuildChildren);
+	}
+
+	UE_API FScopedWidgetSlotArguments InsertSlot(int32 Index = INDEX_NONE, bool bRebuildChildren = true);
+
+	FSlot& GetSlot(int32 SlotIndex)
+	{
+		return Children[SlotIndex];
+	}
+
+	const FSlot& GetSlot(int32 SlotIndex) const
+	{
+		return Children[SlotIndex];
+	}
+
+	int32 NumSlots() const
+	{
+		return Children.Num();
+	}
+
+	void ClearChildren()
+	{
+		Children.Empty();
+	}
+
+	SSegmentedProgressBar()
+		: Children(this)
+	{
+		bCanSupportFocus = false;
+	}
+
+
+
+	UE_API void Construct(const FArguments& Arguments);
+	UE_API virtual void Tick(const FGeometry& AllottedGeometry, const double InCurrentTime, const float InDeltaTime) override;
+
+private:
+	void RebuildChildren();
+	void AddChildSeparatorBar( TSharedPtr<class SHorizontalBox> SlotBox, int32 ChildIndex ) const;
+	TSharedRef<SWidget> ConstructChild( const FSlot& Slot, int32 ChildIndex ) const;
+
+	void CacheChildStates();
+	EState GetChildState( int32 ChildIndex ) const;
+
+	FLinearColor GetSeparatorBarColor( int32 ChildIndex ) const;
+	FSlateColor GetCircleColor( int32 ChildIndex ) const;
+	FSlateColor GetIconColor( int32 ChildIndex ) const;
+	EVisibility GetFullCircleVisibility( int32 ChildIndex ) const;
+	EVisibility GetOuterCircleVisibility( int32 ChildIndex ) const;
+	EVisibility GetProgressCircleVisibility( int32 ChildIndex ) const;
+	EVisibility GetOverlayVisibility( int32 ChildIndex ) const;
+	const FSlateBrush* GetOverlayIcon( int32 ChildIndex ) const;
+	TOptional<FSlateRenderTransform> GetProgressCircleTransform( int32 ChildIndex ) const;
+
+	FSegmentedProgressBarStyle Style;
+	TPanelChildren<FSlot> Children;
+	TArray<EState> CachedChildStates;
+	EState CachedOverallState = EState::None;
+
+	FCurveSequence ThrobberAnimation;
+	const float LineSize = 4;
+};
+
+#undef UE_API
