ToolWidgets: Added DragOverridesCenterAlignment to SDraggableBoxOverlay to allow dragging to override center alignment on either axis. This is useful for when we want a DraggableBoxOverlay to start centered but still allow users to customize its position.

#Synced-CL 52701436

diff --git a/Engine/Source/Developer/ToolWidgets/Private/Overlay/SDraggableBox.cpp b/Engine/Source/Developer/ToolWidgets/Private/Overlay/SDraggableBox.cpp
--- a/Engine/Source/Developer/ToolWidgets/Private/Overlay/SDraggableBox.cpp
+++ b/Engine/Source/Developer/ToolWidgets/Private/Overlay/SDraggableBox.cpp
@@ -56,6 +56,7 @@ void SDraggableBox::Construct(const FArguments& InArgs, const TSharedRef<SDragga
 	DraggableOverlayWeak = InDraggableOverlay;
 	InnerWidget = InArgs._Content.Widget;
 	IsDraggableAttr = InArgs._IsDraggable;
+	DragOverridesCenterAlignmentAttr = InArgs._DragOverridesCenterAlignment;
 	OnUserDraggedToNewPositionDelegate = InArgs._OnUserDraggedToNewPosition;
 
 	ChildSlot
@@ -75,8 +76,18 @@ void SDraggableBox::OnDragUpdate(const FPointerEvent& InMouseEvent, const FDragI
 	const FGeometry& MyGeometry = DraggableOverlay->GetTickSpaceGeometry();
 	const FVector2f MouseOffset = (InMouseEvent.GetScreenSpacePosition() - InDragInfo.OriginalMousePosition)
 		* (MyGeometry.GetLocalSize() / MyGeometry.GetAbsoluteSize());
-
+	
 	FVector2f NewAlignmentOffset = InDragInfo.OriginalAlignmentOffset;
+	EHorizontalAlignment NewHorizontalAlignment = InDragInfo.OriginalHorizontalAlignment;
+	EVerticalAlignment NewVerticalAlignment = InDragInfo.OriginalVerticalAlignment;
+	
+	// If we allow drag to override center alignment, compute the current centered alignment offset.
+	FVector2f CenterOffset(0.f, 0.f);
+	if (DragOverridesCenterAlignmentAttr.Get())
+	{
+		CenterOffset = GetCenterBoxAlignmentOffset();
+	}
+
 	switch (InDragInfo.OriginalHorizontalAlignment)
 	{
 	case EHorizontalAlignment::HAlign_Left:
@@ -86,6 +97,14 @@ void SDraggableBox::OnDragUpdate(const FPointerEvent& InMouseEvent, const FDragI
 	case EHorizontalAlignment::HAlign_Right:
 		NewAlignmentOffset.X -= MouseOffset.X;
 		break;
+		
+	case EHorizontalAlignment::HAlign_Center:
+		if (DragOverridesCenterAlignmentAttr.Get())
+		{
+			NewAlignmentOffset.X = CenterOffset.X + MouseOffset.X;
+			NewHorizontalAlignment = HAlign_Left;
+		}
+		break;
 
 	default:
 		// Do nothing
@@ -101,14 +120,22 @@ void SDraggableBox::OnDragUpdate(const FPointerEvent& InMouseEvent, const FDragI
 	case EVerticalAlignment::VAlign_Bottom:
 		NewAlignmentOffset.Y -= MouseOffset.Y;
 		break;
+		
+	case EVerticalAlignment::VAlign_Center:
+		if (DragOverridesCenterAlignmentAttr.Get())
+		{
+			NewAlignmentOffset.Y = CenterOffset.Y + MouseOffset.Y;
+			NewVerticalAlignment = VAlign_Top;
+		}
+		break;
 
 	default:
 		// Do nothing
 		break;
 	}
 
-	DraggableOverlay->SetBoxHorizontalAlignment(InDragInfo.OriginalHorizontalAlignment);
-	DraggableOverlay->SetBoxVerticalAlignment(InDragInfo.OriginalVerticalAlignment);
+	DraggableOverlay->SetBoxHorizontalAlignment(NewHorizontalAlignment);
+	DraggableOverlay->SetBoxVerticalAlignment(NewVerticalAlignment);
 	DraggableOverlay->SetBoxAlignmentOffset(NewAlignmentOffset);
 
 	if (bInDropped)
@@ -142,4 +169,29 @@ FReply SDraggableBox::OnDragDetected(const FGeometry& InGeometry, const FPointer
 	}
 	return FReply::Unhandled();
 }
+	
+FVector2f SDraggableBox::GetCenterBoxAlignmentOffset() const
+{
+	const TSharedPtr<SDraggableBoxOverlay> DraggableOverlay = DraggableOverlayWeak.Pin();
+	if (!DraggableOverlay.IsValid() || !IsDraggableAttr.Get())
+	{
+		return FVector2f::ZeroVector;
+	}
+	
+	const FGeometry& MyGeometry = DraggableOverlay->GetTickSpaceGeometry();
+	const FVector2f MyGeometryScale = MyGeometry.GetLocalSize() / MyGeometry.GetAbsoluteSize();
+
+	const FVector2f BoxSize = GetTickSpaceGeometry().GetAbsoluteSize();
+	const FVector2f AvailableSpace = (MyGeometry.GetAbsoluteSize() - BoxSize);
+	const FVector2f MidPoint = AvailableSpace * 0.5f * MyGeometryScale;
+	if (!ensure(!MidPoint.ContainsNaN())) // Don't call this during construction because the geometry is not initialized, yet. 
+	{
+		return FVector2f::ZeroVector;
+	}
+	
+	FVector2f CenterOffset;
+	CenterOffset.X = FMath::Max(MidPoint.X, 0.f);
+	CenterOffset.Y = FMath::Max(MidPoint.Y, 0.f);
+	return CenterOffset;
+}
 }
\ No newline at end of file
diff --git a/Engine/Source/Developer/ToolWidgets/Private/Overlay/SDraggableBox.h b/Engine/Source/Developer/ToolWidgets/Private/Overlay/SDraggableBox.h
--- a/Engine/Source/Developer/ToolWidgets/Private/Overlay/SDraggableBox.h
+++ b/Engine/Source/Developer/ToolWidgets/Private/Overlay/SDraggableBox.h
@@ -23,17 +23,21 @@ class SDraggableBox : public SCompoundWidget
 public:
 	SLATE_BEGIN_ARGS(SDraggableBox)
 		: _IsDraggable(true)
+		, _DragOverridesCenterAlignment(false)
 	{}
 	SLATE_DEFAULT_SLOT(FArguments, Content)
 
 	/** Whether this widget should be allowed to be dragged. */
 	SLATE_ATTRIBUTE(bool, IsDraggable)
+
+	/** Whether this widget can override center alignment when dragged. */
+	SLATE_ATTRIBUTE(bool, DragOverridesCenterAlignment)
 	
 	/** Invoked when the user has finished dragging the widget to a new position. */
 	SLATE_EVENT(FSimpleDelegate, OnUserDraggedToNewPosition)
-SLATE_END_ARGS()
+	SLATE_END_ARGS()
 
-void Construct(const FArguments& InArgs, const TSharedRef<SDraggableBoxOverlay>& InDraggableOverlay);
+	void Construct(const FArguments& InArgs, const TSharedRef<SDraggableBoxOverlay>& InDraggableOverlay);
 
 	struct FDragInfo
 	{
@@ -59,8 +63,13 @@ private:
 
 	/** Whether this widget should be allowed to be dragged. */
 	TAttribute<bool> IsDraggableAttr;
+	
+	/** Whether this widget can override center alignment when dragged. */
+	TAttribute<bool> DragOverridesCenterAlignmentAttr;
 
 	/** The content you want to be able to drag around in the parent widget. */
 	FSimpleDelegate OnUserDraggedToNewPositionDelegate;
+	
+	FVector2f GetCenterBoxAlignmentOffset() const;
 };
 }
\ No newline at end of file
diff --git a/Engine/Source/Developer/ToolWidgets/Private/Overlay/SDraggableBoxOverlay.cpp b/Engine/Source/Developer/ToolWidgets/Private/Overlay/SDraggableBoxOverlay.cpp
--- a/Engine/Source/Developer/ToolWidgets/Private/Overlay/SDraggableBoxOverlay.cpp
+++ b/Engine/Source/Developer/ToolWidgets/Private/Overlay/SDraggableBoxOverlay.cpp
@@ -28,6 +28,7 @@ void SDraggableBoxOverlay::Construct(const FArguments& InArgs)
 		[
 			SAssignNew(DraggableBox, SDraggableBox, SharedThis(this))
 			.IsDraggable(InArgs._IsDraggable)
+			.DragOverridesCenterAlignment(InArgs._DragOverridesCenterAlignment)
 			.OnUserDraggedToNewPosition(InArgs._OnUserDraggedToNewPosition)
 			[
 				InArgs._Content.Widget
@@ -230,4 +231,4 @@ void SDraggableBoxOverlay::RestoreFromDragBoxPosition(const FToolWidget_DragBoxP
 	constexpr bool bRecomputeAnchors = false;
 	SetBoxAlignmentOffset(InWidgetPosition.RelativeOffset, bRecomputeAnchors);
 }
-}
\ No newline at end of file
+}
diff --git a/Engine/Source/Developer/ToolWidgets/Public/Overlay/SDraggableBoxOverlay.h b/Engine/Source/Developer/ToolWidgets/Public/Overlay/SDraggableBoxOverlay.h
--- a/Engine/Source/Developer/ToolWidgets/Public/Overlay/SDraggableBoxOverlay.h
+++ b/Engine/Source/Developer/ToolWidgets/Public/Overlay/SDraggableBoxOverlay.h
@@ -42,6 +42,8 @@ public:
 	SLATE_ARGUMENT(FVector2f, InitialAlignmentOffset)
 	/** Whether this widget should be allowed to be dragged. */
 	SLATE_ATTRIBUTE(bool, IsDraggable)
+	/** Whether this widget can override center alignment when dragged. */
+	SLATE_ATTRIBUTE(bool, DragOverridesCenterAlignment)
 
 	/**
 	 * Invoked when the user has finished dragging the widget to a new position.
@@ -51,9 +53,9 @@ public:
 	
 	/** The content you want to be able to drag around in the parent widget. */
 	SLATE_DEFAULT_SLOT(FArguments, Content)
-SLATE_END_ARGS()
+	SLATE_END_ARGS()
 
-UE_API void Construct(const FArguments& InArgs);
+	UE_API void Construct(const FArguments& InArgs);
 
 	/** @return The current position of the box on the viewport relative to HAlign and VAlign. */
 	UE_API FVector2f GetBoxAlignmentOffset() const;
